[MLDeformer] Adding/moving the ML Deformer to the Experimental folder.

diff --git a/Engine/Plugins/Experimental/MLDeformer/MLDeformer.uplugin b/Engine/Plugins/Experimental/MLDeformer/MLDeformer.uplugin
new file mode 100644
--- /dev/null
+++ b/Engine/Plugins/Experimental/MLDeformer/MLDeformer.uplugin
@@ -0,0 +1,43 @@
+{
+	"FileVersion": 3,
+	"Version": 1,
+	"VersionName": "0.1",
+	"FriendlyName": "MLDeformer",
+	"Description": "Machine Learning Mesh Deformer",
+	"Category": "Animation",
+	"CreatedBy": "Epic Games, Inc.",
+	"CreatedByURL": "http://epicgames.com",
+	"DocsURL": "",
+	"MarketplaceURL": "",
+	"SupportURL": "",
+	"CanContainContent": true,
+	"IsBetaVersion": true,
+	"IsExperimentalVersion": true,
+	"Installed": false,
+	"Modules": [
+		{
+			"Name": "MLDeformer",
+			"Type": "Runtime",
+			"LoadingPhase": "Default"
+		},
+		{
+			"Name": "MLDeformerEditor",
+			"Type": "Editor",
+			"LoadingPhase": "Default"
+		}
+	],
+	"Plugins": [
+		{
+			"Name": "GeometryCache",
+			"Enabled": true
+		},
+		{
+			"Name": "NeuralNetworkInference",
+			"Enabled": true
+		},
+		{
+			"Name": "Optimus",
+			"Enabled": true
+		}
+	]
+}
\ No newline at end of file
diff --git a/Engine/Plugins/Experimental/MLDeformer/Shaders/Private/MLDeformerGraphDataInterface.ush b/Engine/Plugins/Experimental/MLDeformer/Shaders/Private/MLDeformerGraphDataInterface.ush
new file mode 100644
--- /dev/null
+++ b/Engine/Plugins/Experimental/MLDeformer/Shaders/Private/MLDeformerGraphDataInterface.ush
@@ -0,0 +1,24 @@
+// Copyright Epic Games, Inc. All Rights Reserved.
+
+#include "/Engine/Private/ComputeKernelCommon.ush"
+
+uint DI_LOCAL(NumVertices);
+float3 DI_LOCAL(VertexDeltaScale);
+float3 DI_LOCAL(VertexDeltaMean);
+float DI_LOCAL(VertexDeltaMultiplier);
+
+Buffer<float> DI_LOCAL(PositionDeltaBuffer);
+Buffer<uint> DI_LOCAL(VertexMapBuffer);
+
+DI_IMPL_READ(ReadNumVertices, uint, void)
+{
+	return DI_LOCAL(NumVertices);
+}
+
+DI_IMPL_READ(ReadPositionDelta, float3, uint VertexIndex)
+{
+	uint BufferIndex = DI_LOCAL(VertexMapBuffer)[VertexIndex] * 3;
+	float3 Delta = float3(DI_LOCAL(PositionDeltaBuffer)[BufferIndex], DI_LOCAL(PositionDeltaBuffer)[BufferIndex + 1], DI_LOCAL(PositionDeltaBuffer)[BufferIndex + 2]);
+	float3 FinalDelta = Delta * DI_LOCAL(VertexDeltaScale) + DI_LOCAL(VertexDeltaMean);
+	return FinalDelta * DI_LOCAL(VertexDeltaMultiplier);
+}
diff --git a/Engine/Plugins/Experimental/MLDeformer/Source/MLDeformer/MLDeformer.Build.cs b/Engine/Plugins/Experimental/MLDeformer/Source/MLDeformer/MLDeformer.Build.cs
new file mode 100644
--- /dev/null
+++ b/Engine/Plugins/Experimental/MLDeformer/Source/MLDeformer/MLDeformer.Build.cs
@@ -0,0 +1,56 @@
+// Copyright Epic Games, Inc. All Rights Reserved.
+
+using UnrealBuildTool;
+
+public class MLDeformer : ModuleRules
+{
+	public MLDeformer(ReadOnlyTargetRules Target) : base(Target)
+	{
+		PCHUsage = ModuleRules.PCHUsageMode.UseExplicitOrSharedPCHs;
+		
+		PublicIncludePaths.AddRange(
+			new string[] {
+				// ... add public include paths required here ...
+			}
+			);
+				
+		
+		PrivateIncludePaths.AddRange(
+			new string[] {
+				// ... add other private include paths required here ...
+			}
+			);
+			
+		
+		PublicDependencyModuleNames.AddRange(
+			new string[]
+			{
+				"Core",
+				// ... add other public dependencies that you statically link with here ...
+			}
+			);
+			
+		
+		PrivateDependencyModuleNames.AddRange(
+			new string[]
+			{
+				"CoreUObject",
+				"Engine",
+				"GeometryCache",
+				"NeuralNetworkInference",
+				"OptimusCore",
+				"Projects",
+				"RenderCore",
+				"RHI"
+			}
+			);
+		
+		
+		DynamicallyLoadedModuleNames.AddRange(
+			new string[]
+			{
+				// ... add any modules that your module loads dynamically here ...
+			}
+			);
+	}
+}
diff --git a/Engine/Plugins/Experimental/MLDeformer/Source/MLDeformer/Private/MLDeformer.cpp b/Engine/Plugins/Experimental/MLDeformer/Source/MLDeformer/Private/MLDeformer.cpp
new file mode 100644
--- /dev/null
+++ b/Engine/Plugins/Experimental/MLDeformer/Source/MLDeformer/Private/MLDeformer.cpp
@@ -0,0 +1,27 @@
+// Copyright Epic Games, Inc. All Rights Reserved.
+
+#include "MLDeformer.h"
+
+#include "Modules/ModuleManager.h"
+#include "Interfaces/IPluginManager.h"
+#include "Misc/Paths.h"
+#include "ShaderCore.h"
+
+
+MLDEFORMER_API DEFINE_LOG_CATEGORY(LogMLDeformer)
+
+#define LOCTEXT_NAMESPACE "MLDeformerModule"
+
+IMPLEMENT_MODULE(FMLDeformerModule, MLDeformer)
+
+void FMLDeformerModule::StartupModule()
+{
+	const FString PluginShaderDir = FPaths::Combine(IPluginManager::Get().FindPlugin(TEXT("MLDeformer"))->GetBaseDir(), TEXT("Shaders"));
+	AddShaderSourceDirectoryMapping(TEXT("/Plugin/MLDeformer"), PluginShaderDir);
+}
+
+void FMLDeformerModule::ShutdownModule()
+{
+}
+
+#undef LOCTEXT_NAMESPACE
diff --git a/Engine/Plugins/Experimental/MLDeformer/Source/MLDeformer/Private/MLDeformerAsset.cpp b/Engine/Plugins/Experimental/MLDeformer/Source/MLDeformer/Private/MLDeformerAsset.cpp
new file mode 100644
--- /dev/null
+++ b/Engine/Plugins/Experimental/MLDeformer/Source/MLDeformer/Private/MLDeformerAsset.cpp
@@ -0,0 +1,365 @@
+// Copyright Epic Games, Inc. All Rights Reserved.
+
+#include "MLDeformerAsset.h"
+#include "MLDeformer.h"
+#include "MLDeformerVizSettings.h"
+#include "MLDeformerInputInfo.h"
+#include "GeometryCache.h"
+#include "GeometryCacheMeshData.h"
+#include "Rendering/SkeletalMeshModel.h"
+#include "Rendering/SkeletalMeshLODModel.h"
+#include "NeuralNetwork.h"
+#include "Engine/SkeletalMesh.h"
+#include "Animation/AnimSequence.h"
+
+#define LOCTEXT_NAMESPACE "MLDeformerAsset"
+
+void FVertexMapBuffer::InitRHI()
+{
+	if (VertexMap.Num() > 0)
+	{
+		FRHIResourceCreateInfo CreateInfo(TEXT("FVertexMapBuffer"));
+
+		VertexBufferRHI = RHICreateVertexBuffer(VertexMap.Num() * sizeof(uint32), BUF_Static | BUF_ShaderResource, CreateInfo);
+		uint32* Data = reinterpret_cast<uint32*>(RHILockBuffer(VertexBufferRHI, 0, VertexMap.Num() * sizeof(uint32), RLM_WriteOnly));
+		for (int32 Index = 0; Index < VertexMap.Num(); ++Index)
+		{
+			Data[Index] = static_cast<uint32>(VertexMap[Index]);
+		}
+		RHIUnlockBuffer(VertexBufferRHI);
+		VertexMap.Empty();
+
+		ShaderResourceViewRHI = RHICreateShaderResourceView(VertexBufferRHI, 4, PF_R32_UINT);
+	}
+	else
+	{
+		VertexBufferRHI = nullptr;
+		ShaderResourceViewRHI = nullptr;
+	}
+}
+
+UMLDeformerAsset::UMLDeformerAsset(const FObjectInitializer& ObjectInitializer)
+	: Super(ObjectInitializer)
+{
+#if WITH_EDITORONLY_DATA
+	VizSettings = ObjectInitializer.CreateEditorOnlyDefaultSubobject<UMLDeformerVizSettings>(this, TEXT("VizSettings"));
+#endif
+}
+
+void UMLDeformerAsset::Serialize(FArchive& Archive)
+{
+#if WITH_EDITOR
+	if (Archive.IsSaving() && Archive.IsPersistent())
+	{
+		InitVertexMap();
+		UpdateCachedNumVertices();
+	}
+#endif
+
+	Super::Serialize(Archive);
+}
+
+void UMLDeformerAsset::PostLoad()
+{
+	Super::PostLoad();
+	if (NeuralNetwork)
+	{
+		NeuralNetwork->SetBackEnd(ENeuralBackEnd::UEOnly);
+		NeuralNetwork->SetInputDeviceType(ENeuralDeviceType::CPU);
+		NeuralNetwork->SetDeviceType(ENeuralDeviceType::GPU);
+		NeuralNetwork->SetOutputDeviceType(ENeuralDeviceType::GPU);
+	}
+
+#if WITH_EDITOR
+	UpdateCachedNumVertices();
+
+	// Initialize the input info if we didn't store any yet.
+	// This is just for backward compatibility. Normally this data is always there.
+	if (InputInfo.GetNumBones() == 0 && InputInfo.GetNumCurves() == 0)
+	{
+		InputInfo = CreateInputInfo();
+	}
+	else
+#endif
+	{
+		InputInfo.UpdateFNames();
+	}
+
+	InitGPUData();
+}
+
+#if WITH_EDITOR
+FText UMLDeformerAsset::GetBaseAssetChangedErrorText() const
+{
+	FText Result;
+
+	if (SkeletalMesh)
+	{
+		if (NumSkeletalMeshVerts != GetInputInfo().GetNumBaseMeshVertices() &&
+			NumSkeletalMeshVerts > 0 && GetInputInfo().GetNumBaseMeshVertices() > 0)
+		{
+			Result = FText::Format(LOCTEXT("BaseMeshMismatch", "Number of vertices in base mesh has changed from {0} to {1} vertices since this ML Deformer Asset was saved! {2}"),
+				GetInputInfo().GetNumBaseMeshVertices(),
+				NumSkeletalMeshVerts,
+				NeuralNetwork ? LOCTEXT("BaseMeshMismatchNN", "Neural network needs to be retrained!") : FText());
+		}
+	}
+
+	return Result;
+}
+
+FText UMLDeformerAsset::GetTargetAssetChangedErrorText() const
+{
+	FText Result;
+
+	if (GeometryCache)
+	{
+		if (NumGeomCacheVerts != GetInputInfo().GetNumTargetMeshVertices() &&
+			NumGeomCacheVerts > 0 && GetInputInfo().GetNumTargetMeshVertices() > 0)
+		{
+			Result = FText::Format(LOCTEXT("TargetMeshMismatch", "Number of vertices in target mesh has changed from {0} to {1} vertices since this ML Deformer Asset was saved! {2}"),
+				GetInputInfo().GetNumTargetMeshVertices(),
+				NumGeomCacheVerts,
+				NeuralNetwork ? LOCTEXT("BaseMeshMismatchNN", "Neural network needs to be retrained!") : FText());
+		}
+	}
+
+	return Result;
+}
+
+void UMLDeformerAsset::PostEditChangeProperty(FPropertyChangedEvent& PropertyChangedEvent)
+{
+	static const FName SkeletalMeshName = GET_MEMBER_NAME_CHECKED(UMLDeformerAsset, SkeletalMesh);
+	const FName PropertyName = PropertyChangedEvent.Property->GetFName();
+	if (PropertyName == SkeletalMeshName)
+	{
+		InitVertexMap();
+		InitGPUData();
+	}
+
+	Super::PostEditChangeProperty(PropertyChangedEvent);
+}
+
+int32 UMLDeformerAsset::GetNumFrames() const
+{
+	check(GeometryCache != nullptr)
+	const int32 StartFrame = GeometryCache->GetStartFrame();
+	const int32 EndFrame = GeometryCache->GetEndFrame();
+	check(EndFrame >= StartFrame);
+	return (EndFrame - StartFrame) + 1;
+}
+#endif
+
+void UMLDeformerAsset::BeginDestroy()
+{
+	Super::BeginDestroy();
+	BeginReleaseResource(&VertexMapBuffer);
+	RenderResourceDestroyFence.BeginFence();
+}
+
+bool UMLDeformerAsset::IsReadyForFinishDestroy()
+{
+	return Super::IsReadyForFinishDestroy() && RenderResourceDestroyFence.IsFenceComplete();
+}
+
+#if WITH_EDITOR
+void UMLDeformerAsset::InitVertexMap()
+{
+	VertexMap.Empty();
+	if (SkeletalMesh)
+	{
+		FSkeletalMeshModel* SkeletalMeshModel = SkeletalMesh->GetImportedModel();
+		if (SkeletalMeshModel)
+		{
+			VertexMap = SkeletalMeshModel->LODModels[0].MeshToImportVertexMap;
+		}
+	}
+}
+
+FText UMLDeformerAsset::GetGeomCacheErrorText(UGeometryCache* InGeomCache) const
+{
+	FText Result;
+	if (InGeomCache)
+	{
+		// Verify that we have imported vertex numbers enabled.
+		TArray<FGeometryCacheMeshData> MeshData;
+		InGeomCache->GetMeshDataAtTime(0.0f, MeshData);
+		if (MeshData.Num() == 0)
+		{
+			Result = LOCTEXT("TargetMeshNoMeshData", "No geometry data is present.");
+		}
+		else
+		{
+			if (MeshData[0].ImportedVertexNumbers.Num() == 0)
+			{
+				Result = LOCTEXT("TargetMeshNoImportedVertexNumbers", "Please import Geometry Cache with option 'Store Imported Vertex Numbers' enabled!");
+			}
+		}
+	}
+
+	return Result;
+}
+
+FText UMLDeformerAsset::GetVertexErrorText(USkeletalMesh* InSkelMesh, UGeometryCache* InGeomCache, const FText& SkelName, const FText& GeomCacheName) const
+{
+	FText Result;
+
+	if (InSkelMesh && InGeomCache)
+	{
+		FSkeletalMeshModel* SkeletalMeshModel = InSkelMesh->GetImportedModel();	
+		TArray<FGeometryCacheMeshData> MeshData;
+		InGeomCache->GetMeshDataAtTime(0.0f, MeshData);
+		if (SkeletalMeshModel && MeshData.Num() > 0)
+		{
+			const TArray<int32>& SkeletalVertexMap = SkeletalMeshModel->LODModels[0].MeshToImportVertexMap;
+			const TArray<uint32>& GeomCacheVertexMap = MeshData[0].ImportedVertexNumbers;
+			if (!GeomCacheVertexMap.IsEmpty() && !SkeletalVertexMap.IsEmpty())
+			{
+				// Find the max vertex index.
+				int32 MaxGeomCacheVertex = 0;
+				for (int32 Index = 0; Index < GeomCacheVertexMap.Num(); ++Index)
+				{
+					MaxGeomCacheVertex = FMath::Max<int32>(static_cast<int32>(GeomCacheVertexMap[Index]), MaxGeomCacheVertex);
+				}
+
+				const int32 MaxSkelMeshVertex = SkeletalMeshModel->LODModels[0].MaxImportVertex;
+				if (MaxSkelMeshVertex != MaxGeomCacheVertex)
+				{
+					Result = FText::Format(
+						LOCTEXT("MeshVertexNumVertsMismatch", "Vertex count of {0} doesn't match with {1}!\n\n{2} has {3} verts, while {4} has {5} verts."),
+						SkelName,
+						GeomCacheName,
+						SkelName,
+						MaxSkelMeshVertex + 1,
+						GeomCacheName,
+						MaxGeomCacheVertex + 1);
+				}
+			}
+		}
+	}
+
+	return Result;
+}
+
+FText UMLDeformerAsset::GetAnimSequenceErrorText(UGeometryCache* InGeomCache, UAnimSequence* InAnimSequence) const
+{
+	FText Result;
+	if (InAnimSequence && InGeomCache)
+	{
+		const float AnimSeqDuration = InAnimSequence->GetPlayLength();
+		const float GeomCacheDuration = InGeomCache->CalculateDuration();
+		if (FMath::Abs(AnimSeqDuration - GeomCacheDuration) > 0.001f)
+		{
+			FNumberFormattingOptions Options;
+			Options.SetUseGrouping(false);
+			Options.SetMaximumFractionalDigits(4);
+			Result = FText::Format(
+				LOCTEXT("AnimSeqNumFramesMismatch", "Anim sequence and Geometry Cache durations don't match!\n\nAnimSeq has a duration of {0} seconds, while GeomCache has a duration of {1} seconds.\n\nThis can produce incorrect results."),
+				FText::AsNumber(AnimSeqDuration, &Options),
+				FText::AsNumber(GeomCacheDuration, &Options));
+		}
+	}
+	return Result;
+}
+
+FText UMLDeformerAsset::GetInputsErrorText() const
+{
+	if (SkeletalMesh && CreateInputInfo().IsEmpty())
+	{
+		switch (TrainingInputs)			
+		{
+			case ETrainingInputs::BonesOnly:		return FText(LOCTEXT("InputsEmptyBonesErrorText", "Your base mesh has no bones to train on."));
+			case ETrainingInputs::CurvesOnly:		return FText(LOCTEXT("InputsEmptyCurvesErrorText", "Your base mesh has no curves to train on."));
+			case ETrainingInputs::BonesAndCurves:	return FText(LOCTEXT("InputsEmptyBonesCurvesErrorText", "Your base mesh has no bones or curves to train on."));
+			default: return FText(LOCTEXT("InputsEmptyDefaultErrorText", "There are no inputs to train on. There are no bones, curves or other inputs we can use."));
+		}
+	}
+
+	return FText();
+}
+
+FText UMLDeformerAsset::GetIncompatibleSkeletonErrorText(USkeletalMesh* InSkelMesh, UAnimSequence* InAnimSeq) const
+{
+	FText Result;
+	if (InSkelMesh && InAnimSeq)
+	{
+		if (!InSkelMesh->GetSkeleton()->IsCompatible(InAnimSeq->GetSkeleton()))
+		{
+			Result = LOCTEXT("SkeletonMismatch", "The base skeletal mesh and anim sequence use different skeletons. The animation might not play correctly.");
+		}
+	}
+	return Result;
+}
+
+bool UMLDeformerAsset::IsCompatibleWithNeuralNet() const
+{	
+	if (SkeletalMesh == nullptr)
+	{
+		return true;
+	}
+
+	return GetInputInfo().IsCompatible(SkeletalMesh);
+}
+
+void UMLDeformerAsset::UpdateCachedNumVertices()
+{
+	NumSkeletalMeshVerts = 0;
+	NumGeomCacheVerts = 0;
+
+	// Extract max num vertices from the geometry cache.
+	if (GeometryCache)
+	{
+		TArray<FGeometryCacheMeshData> MeshData;
+		GeometryCache->GetMeshDataAtTime(0.0f, MeshData);
+		if (MeshData.Num() > 0)
+		{
+			const TArray<uint32>& VertexNumbers = MeshData[0].ImportedVertexNumbers;
+			for (int32 Index = 0; Index < VertexNumbers.Num(); ++Index)
+			{
+				NumGeomCacheVerts = FMath::Max<int32>(static_cast<int32>(VertexNumbers[Index]), NumGeomCacheVerts);
+			}
+
+			// Add one as they are indices starting from 0.
+			if (VertexNumbers.Num() > 0)
+			{
+				NumGeomCacheVerts += 1;
+			}
+		}
+	}
+
+	// Extract max num vertices from the skeletal mesh.
+	if (SkeletalMesh)
+	{
+		FSkeletalMeshModel* SkeletalMeshModel = SkeletalMesh->GetImportedModel();
+		if (SkeletalMeshModel)
+		{
+			const int32 MaxIndex = SkeletalMeshModel->LODModels[0].MaxImportVertex;
+			NumSkeletalMeshVerts = (MaxIndex > 0) ? (MaxIndex + 1) : 0;
+		}
+	}
+}
+
+FMLDeformerInputInfo UMLDeformerAsset::CreateInputInfo() const
+{
+	FMLDeformerInputInfo Result;
+
+	// Init the inputs.
+	FMLDeformerInputInfoInitSettings Settings;
+	Settings.SkeletalMesh = SkeletalMesh;
+	Settings.TargetMesh = GeometryCache;
+	Settings.bIncludeBones = (TrainingInputs == ETrainingInputs::BonesAndCurves || TrainingInputs == ETrainingInputs::BonesOnly);
+	Settings.bIncludeCurves = (TrainingInputs == ETrainingInputs::BonesAndCurves || TrainingInputs == ETrainingInputs::CurvesOnly);
+	Result.Init(Settings);
+
+	return Result;
+}
+#endif
+
+void UMLDeformerAsset::InitGPUData()
+{
+	BeginReleaseResource(&VertexMapBuffer);
+	VertexMapBuffer.Init(VertexMap);
+	BeginInitResource(&VertexMapBuffer);
+}
+
+#undef LOCTEXT_NAMESPACE
diff --git a/Engine/Plugins/Experimental/MLDeformer/Source/MLDeformer/Private/MLDeformerComponent.cpp b/Engine/Plugins/Experimental/MLDeformer/Source/MLDeformer/Private/MLDeformerComponent.cpp
new file mode 100644
--- /dev/null
+++ b/Engine/Plugins/Experimental/MLDeformer/Source/MLDeformer/Private/MLDeformerComponent.cpp
@@ -0,0 +1,47 @@
+// Copyright Epic Games, Inc. All Rights Reserved.
+
+#include "MLDeformerComponent.h"
+#include "MLDeformer.h"
+#include "MLDeformerAsset.h"
+#include "MLDeformerInstance.h"
+#include "Components/SkeletalMeshComponent.h"
+
+UMLDeformerComponent::UMLDeformerComponent(const FObjectInitializer& ObjectInitializer)
+	: Super(ObjectInitializer)
+{
+	bTickInEditor = true;
+	bAutoActivate = true;
+	PrimaryComponentTick.TickGroup = TG_PrePhysics;
+	PrimaryComponentTick.bCanEverTick = true;
+}
+
+void UMLDeformerComponent::SetupComponent(UMLDeformerAsset* InDeformerAsset, USkeletalMeshComponent* InSkelMeshComponent)
+{
+	DeformerAsset = InDeformerAsset;
+	SkelMeshComponent = InSkelMeshComponent;
+	DeformerInstance.Init(InDeformerAsset, InSkelMeshComponent);
+}
+
+void UMLDeformerComponent::Activate(bool bReset)
+{
+	// If we haven't pointed to some skeletal mesh component to use, then try to find one on the actor.
+	USkeletalMeshComponent* SkelMeshComponentToUse = SkelMeshComponent;
+	if (SkelMeshComponentToUse == nullptr)
+	{
+		AActor* Actor = Cast<AActor>(GetOuter());
+		SkelMeshComponent = Actor->FindComponentByClass<USkeletalMeshComponent>();
+	}
+
+	SetupComponent(DeformerAsset, SkelMeshComponent);
+}
+
+void UMLDeformerComponent::TickComponent(float DeltaTime, enum ELevelTick TickType, FActorComponentTickFunction* ThisTickFunction)
+{
+	if (TickType != ELevelTick::LEVELTICK_PauseTick)
+	{
+		// Update the deformer, which runs the inference.
+		DeformerInstance.Update();
+	}
+
+	Super::TickComponent(DeltaTime, TickType, ThisTickFunction);
+}
diff --git a/Engine/Plugins/Experimental/MLDeformer/Source/MLDeformer/Private/MLDeformerGraphDataInterface.cpp b/Engine/Plugins/Experimental/MLDeformer/Source/MLDeformer/Private/MLDeformerGraphDataInterface.cpp
new file mode 100644
--- /dev/null
+++ b/Engine/Plugins/Experimental/MLDeformer/Source/MLDeformer/Private/MLDeformerGraphDataInterface.cpp
@@ -0,0 +1,146 @@
+// Copyright Epic Games, Inc. All Rights Reserved.
+
+#include "MLDeformerGraphDataInterface.h"
+#include "MLDeformerAsset.h"
+#include "MLDeformerComponent.h"
+
+#include "ComputeFramework/ComputeKernelPermutationSet.h"
+#include "ComputeFramework/ShaderParamTypeDefinition.h"
+#include "NeuralNetwork.h"
+#include "RenderGraphBuilder.h"
+#include "RenderGraphUtils.h"
+#include "RenderGraphResources.h"
+#include "ShaderParameterMetadataBuilder.h"
+
+FString UMLDeformerDataInterface::GetDisplayName() const
+{
+	return TEXT("ML Deformer");
+}
+
+TArray<FOptimusCDIPinDefinition> UMLDeformerDataInterface::GetPinDefinitions() const
+{
+	TArray<FOptimusCDIPinDefinition> Defs;
+	Defs.Add({ "PositionDelta", "ReadPositionDelta", "ReadNumVertices", "Vertex" });
+	return Defs;
+}
+
+void UMLDeformerDataInterface::GetSupportedInputs(TArray<FShaderFunctionDefinition>& OutFunctions) const
+{
+	{
+		FShaderFunctionDefinition Fn;
+		Fn.Name = TEXT("ReadNumVertices");
+		Fn.bHasReturnType = true;
+		FShaderParamTypeDefinition ReturnParam = {};
+		ReturnParam.ValueType = FShaderValueType::Get(EShaderFundamentalType::Uint);
+		Fn.ParamTypes.Add(ReturnParam);
+		OutFunctions.Add(Fn);
+	}
+	{
+		FShaderFunctionDefinition Fn;
+		Fn.Name = TEXT("ReadPositionDelta");
+		Fn.bHasReturnType = true;
+		FShaderParamTypeDefinition ReturnParam = {};
+		ReturnParam.ValueType = FShaderValueType::Get(EShaderFundamentalType::Float, 3);
+		Fn.ParamTypes.Add(ReturnParam);
+		FShaderParamTypeDefinition Param0 = {};
+		Param0.ValueType = FShaderValueType::Get(EShaderFundamentalType::Uint);
+		Fn.ParamTypes.Add(Param0);
+		OutFunctions.Add(Fn);
+	}
+}
+
+BEGIN_SHADER_PARAMETER_STRUCT(FSceneDataInterfaceParameters, )
+	SHADER_PARAMETER(uint32, NumVertices)
+	SHADER_PARAMETER(FVector3f, VertexDeltaScale)
+	SHADER_PARAMETER(FVector3f, VertexDeltaMean)
+	SHADER_PARAMETER(float, VertexDeltaMultiplier)
+	SHADER_PARAMETER_RDG_BUFFER_SRV(Buffer<float>, PositionDeltaBuffer)
+	SHADER_PARAMETER_SRV(Buffer<uint>, VertexMapBuffer)
+END_SHADER_PARAMETER_STRUCT()
+
+void UMLDeformerDataInterface::GetShaderParameters(TCHAR const* UID, FShaderParametersMetadataBuilder& OutBuilder) const
+{
+	OutBuilder.AddNestedStruct<FSceneDataInterfaceParameters>(UID);
+}
+
+void UMLDeformerDataInterface::GetHLSL(FString& OutHLSL) const
+{
+	OutHLSL += TEXT("#include \"/Plugin/MLDeformer/Private/MLDeformerGraphDataInterface.ush\"\n");
+}
+
+void UMLDeformerDataInterface::GetSourceTypes(TArray<UClass*>& OutSourceTypes) const
+{
+	OutSourceTypes.Add(UMLDeformerComponent::StaticClass());
+}
+
+UComputeDataProvider* UMLDeformerDataInterface::CreateDataProvider(UObject* InOuter, TArrayView<TObjectPtr<UObject>> InSourceObjects) const
+{
+	UMLDeformerDataProvider* Provider = NewObject<UMLDeformerDataProvider>(InOuter);
+	if (InSourceObjects.Num() == 1)
+	{
+		Provider->DeformerComponent = Cast<UMLDeformerComponent>(InSourceObjects[0]);
+	}
+	return Provider;
+}
+
+bool UMLDeformerDataProvider::IsValid() const
+{
+	return
+		DeformerComponent != nullptr &&
+		DeformerComponent->GetDeformerAsset() != nullptr &&
+		DeformerComponent->GetDeformerAsset()->GetVertexMapBuffer().ShaderResourceViewRHI != nullptr;
+}
+
+FComputeDataProviderRenderProxy* UMLDeformerDataProvider::GetRenderProxy()
+{
+	return new FMLDeformerDataProviderProxy(DeformerComponent);
+}
+
+FMLDeformerDataProviderProxy::FMLDeformerDataProviderProxy(UMLDeformerComponent* DeformerComponent)
+{
+	const UMLDeformerAsset* DeformerAsset = DeformerComponent->GetDeformerAsset();
+	NeuralNetwork = DeformerAsset->GetInferenceNeuralNetwork();
+	VertexMapBufferSRV = DeformerAsset->GetVertexMapBuffer().ShaderResourceViewRHI;
+	VertexDeltaScale = DeformerAsset->GetVertexDeltaScale();
+	VertexDeltaMean = DeformerAsset->GetVertexDeltaMean();
+	bCanRunNeuralNet = NeuralNetwork ? (NeuralNetwork->GetInputTensor().Num() == static_cast<int64>(DeformerAsset->GetInputInfo().CalcNumNeuralNetInputs())) : false;
+
+#ifdef WITH_EDITOR
+	VertexDeltaMultiplier = DeformerComponent->GetVertexDeltaMultiplier();
+#else
+	VertexDeltaMultiplier = 1.0f;
+#endif
+}
+
+void FMLDeformerDataProviderProxy::AllocateResources(FRDGBuilder& GraphBuilder)
+{
+	if (bCanRunNeuralNet)
+	{
+		check(NeuralNetwork != nullptr);
+		Buffer = GraphBuilder.RegisterExternalBuffer(NeuralNetwork->GetOutputTensor().GetPooledBuffer());
+	}
+	else
+	{
+		// TODO: use an actual buffer that is of the right size, and filled with zero's.
+		Buffer = GraphBuilder.RegisterExternalBuffer(GWhiteVertexBufferWithRDG->Buffer);
+	}
+
+	BufferSRV = GraphBuilder.CreateSRV(Buffer, PF_R32_FLOAT);
+}
+
+void FMLDeformerDataProviderProxy::GetBindings(int32 InvocationIndex, TCHAR const* UID, FBindings& OutBindings) const
+{
+	FSceneDataInterfaceParameters Parameters;
+	FMemory::Memset(&Parameters, 0, sizeof(Parameters));
+	Parameters.NumVertices = 0;
+	Parameters.VertexDeltaScale = VertexDeltaScale;
+	Parameters.VertexDeltaMean = VertexDeltaMean;
+	Parameters.VertexDeltaMultiplier = VertexDeltaMultiplier;
+	Parameters.PositionDeltaBuffer = BufferSRV;
+	Parameters.VertexMapBuffer = VertexMapBufferSRV;
+
+	TArray<uint8> ParamData;
+	ParamData.SetNum(sizeof(Parameters));
+	FMemory::Memcpy(ParamData.GetData(), &Parameters, sizeof(Parameters));
+	OutBindings.Structs.Add(TTuple<FString, TArray<uint8>>(UID, MoveTemp(ParamData)));
+}
diff --git a/Engine/Plugins/Experimental/MLDeformer/Source/MLDeformer/Private/MLDeformerInputInfo.cpp b/Engine/Plugins/Experimental/MLDeformer/Source/MLDeformer/Private/MLDeformerInputInfo.cpp
new file mode 100644
--- /dev/null
+++ b/Engine/Plugins/Experimental/MLDeformer/Source/MLDeformer/Private/MLDeformerInputInfo.cpp
@@ -0,0 +1,255 @@
+// Copyright Epic Games, Inc. All Rights Reserved.
+
+#include "MLDeformerInputInfo.h"
+#include "Engine/SkeletalMesh.h"
+#include "Animation/AnimCurveTypes.h"
+#include "Animation/Skeleton.h"
+#include "Animation/AnimInstance.h"
+#include "ReferenceSkeleton.h"
+#include "GeometryCache.h"
+#include "GeometryCacheMeshData.h"
+#include "Rendering/SkeletalMeshModel.h"
+#include "Rendering/SkeletalMeshLODModel.h"
+
+void FMLDeformerInputInfo::Init(const FMLDeformerInputInfoInitSettings& Settings)
+{
+	// Reset things.
+	BoneNameStrings.Empty();
+	BoneNames.Empty();
+	CurveNameStrings.Empty();
+	CurveNames.Empty();
+	NumBaseMeshVertices = 0;
+	NumTargetMeshVertices = 0;
+
+	// Make sure we have all required data.
+	USkeletalMesh* SkeletalMesh = Settings.SkeletalMesh;
+	UGeometryCache* GeomCache = Settings.TargetMesh;
+
+#if WITH_EDITOR
+	if (SkeletalMesh)
+	{
+		// Get the number of imported vertices of the skeletal mesh.
+		FSkeletalMeshModel* SkeletalMeshModel = SkeletalMesh->GetImportedModel();	
+		if (SkeletalMeshModel)
+		{
+			if (!SkeletalMeshModel->LODModels[0].MeshToImportVertexMap.IsEmpty())
+			{
+				NumBaseMeshVertices = SkeletalMeshModel->LODModels[0].MaxImportVertex + 1;
+			}
+		}
+	}
+
+	if (GeomCache)
+	{
+		// Get the number of imported vertices of the geometry cache.
+		TArray<FGeometryCacheMeshData> MeshData;
+		GeomCache->GetMeshDataAtTime(0.0f, MeshData);
+		if (MeshData.Num() > 0)
+		{
+			int32 MaxGeomCacheVertex = 0;
+			const TArray<uint32>& GeomCacheVertexMap = MeshData[0].ImportedVertexNumbers;
+			if (!GeomCacheVertexMap.IsEmpty())
+			{
+				for (int32 Index = 0; Index < GeomCacheVertexMap.Num(); ++Index)
+				{
+					MaxGeomCacheVertex = FMath::Max<int32>(static_cast<int32>(GeomCacheVertexMap[Index]), MaxGeomCacheVertex);
+				}
+				NumTargetMeshVertices = MaxGeomCacheVertex + 1;
+			}
+		}
+	}
+#endif
+
+	// Handle bones.
+	if (Settings.bIncludeBones && SkeletalMesh)
+	{
+		// Include all the bones when no list was provided.
+		const FReferenceSkeleton& RefSkeleton = SkeletalMesh->GetRefSkeleton();
+		if (Settings.BoneNamesToInclude.IsEmpty())
+		{
+			// Grab all bone names.
+			const int32 NumBones = RefSkeleton.GetNum();
+			BoneNameStrings.AddDefaulted(NumBones);
+			for (int32 Index = 0; Index < NumBones; ++Index)
+			{
+				const FName BoneName = RefSkeleton.GetBoneName(Index);
+				BoneNameStrings[Index] = BoneName.ToString();
+			}
+		}
+		else // A list of bones to include was provided.
+		{
+			BoneNameStrings = Settings.BoneNamesToInclude;
+
+			// Remove bones that don't exist.
+			for (int32 Index = 0; Index < BoneNames.Num();)
+			{
+				const FName BoneName(BoneNameStrings[Index]);
+				if (RefSkeleton.FindBoneIndex(BoneName) == INDEX_NONE)
+				{
+					UE_LOG(LogMLDeformer, Warning, TEXT("Bone '%s' in the bones include list doesn't exist, ignoring it."), *BoneNameStrings[Index]);
+					BoneNameStrings.Remove(BoneNameStrings[Index]);
+				}
+				else
+				{
+					Index++;
+				}
+			}
+		}
+	}
+
+	// Handle curves.
+	if (Settings.bIncludeCurves && SkeletalMesh)
+	{
+		// Anim curves.
+		USkeleton* Skeleton = SkeletalMesh->GetSkeleton();
+		const FSmartNameMapping* SmartNameMapping = Skeleton ? Skeleton->GetSmartNameContainer(USkeleton::AnimCurveMappingName) : nullptr;
+		if (SmartNameMapping) // When there are curves.
+		{
+			// Include all curves when no list was provided.
+			if (Settings.CurveNamesToInclude.IsEmpty())
+			{
+				SmartNameMapping->FillNameArray(CurveNames);
+				CurveNameStrings.Reserve(CurveNames.Num());
+				for (FName Name : CurveNames)
+				{
+					CurveNameStrings.Add(Name.ToString());
+				}
+			}
+			else // A list of curve names was provided.
+			{
+				CurveNameStrings = Settings.CurveNamesToInclude;
+
+				// Remove curves that don't exist.
+				for (int32 Index = 0; Index < CurveNameStrings.Num();)
+				{
+					FName CurveName(CurveNameStrings[Index]);
+					if (!SmartNameMapping->Exists(CurveName))
+					{
+						UE_LOG(LogMLDeformer, Warning, TEXT("Anim curve '%s' doesn't exist, ignoring it."), *CurveNameStrings[Index]);
+						CurveNameStrings.RemoveAt(Index);
+					}
+					else
+					{
+						Index++;
+					}
+				}
+			}
+		}
+	}
+
+	// Update the FName arrays.
+	UpdateFNames();
+}
+
+void FMLDeformerInputInfo::UpdateFNames()
+{
+	// Update the bone names.
+	const int32 NumBones = BoneNameStrings.Num();
+	BoneNames.Reset(NumBones);
+	BoneNames.Reserve(NumBones);
+	for (const FString& NameString : BoneNameStrings)
+	{
+		BoneNames.Add(FName(NameString));
+	}
+
+	// Update the curve names.
+	const int32 NumCurves = CurveNameStrings.Num();
+	CurveNames.Reset(NumCurves);
+	CurveNames.Reserve(NumCurves);
+	for (const FString& NameString : CurveNameStrings)
+	{
+		CurveNames.Add(FName(NameString));
+	}
+}
+
+bool FMLDeformerInputInfo::IsCompatible(USkeletalMesh* SkeletalMesh) const
+{
+	if (SkeletalMesh == nullptr)
+	{
+		return false;
+	}
+
+	// Verify that all required bones are there.
+	const FReferenceSkeleton& RefSkeleton = SkeletalMesh->GetRefSkeleton();
+	for (const FName BoneName : BoneNames)
+	{
+		const int32 BoneIndex = RefSkeleton.FindBoneIndex(BoneName);
+		if (BoneIndex == INDEX_NONE) // We're missing the required bone. The network needs to input the transform for this bone.
+		{
+			return false;
+		}
+	}
+
+	// Verify that all required curves are there.
+	USkeleton* Skeleton = SkeletalMesh->GetSkeleton();
+	const FSmartNameMapping* SmartNameMapping = Skeleton ? Skeleton->GetSmartNameContainer(USkeleton::AnimCurveMappingName) : nullptr;
+	if (SmartNameMapping)
+	{
+		for (const FName CurveName : CurveNames)
+		{
+			if (!SmartNameMapping->Exists(CurveName))
+			{
+				return false;
+			}
+		}
+	}
+
+	return true;
+}
+
+FString FMLDeformerInputInfo::GenerateCompatibilityErrorString(USkeletalMesh* SkeletalMesh) const
+{
+	if (SkeletalMesh == nullptr)
+	{
+		return FString();
+	}
+
+	// Verify that all required bones are there.
+	const FReferenceSkeleton& RefSkeleton = SkeletalMesh->GetRefSkeleton();
+	FString ErrorString;
+	for (const FName BoneName : BoneNames)
+	{
+		const int32 BoneIndex = RefSkeleton.FindBoneIndex(BoneName);
+		if (BoneIndex == INDEX_NONE) // We're missing the required bone. The network needs to input the transform for this bone.
+		{
+			ErrorString += FString::Format(TEXT("Required bone '{0}' is missing.\n"), {*BoneName.ToString()});
+		}
+	}
+
+	// Verify that all required curves are there.
+	USkeleton* Skeleton = SkeletalMesh->GetSkeleton();
+	const FSmartNameMapping* SmartNameMapping = Skeleton ? Skeleton->GetSmartNameContainer(USkeleton::AnimCurveMappingName) : nullptr;
+	if (SmartNameMapping)
+	{
+		for (const FName CurveName : CurveNames)
+		{
+			if (!SmartNameMapping->Exists(CurveName))
+			{
+				ErrorString += FString::Format(TEXT("Required curve '{0}' is missing.\n"), {*CurveName.ToString()});
+			}
+		}
+	}
+
+	return ErrorString;
+}
+
+void FMLDeformerInputInfo::ExtractCurveValues(UAnimInstance* AnimInstance, TArray<float>& OutValues) const
+{
+	check(CurveNames.Num() == CurveNameStrings.Num());
+	
+	const int32 NumCurves = CurveNames.Num();
+	OutValues.Reset(NumCurves);
+	OutValues.AddUninitialized(NumCurves);
+	for (int32 Index = 0; Index < NumCurves; ++Index)
+	{
+		const FName CurveName = CurveNames[Index];
+		OutValues[Index] = AnimInstance->GetCurveValue(CurveName);
+	}
+}
+
+int32 FMLDeformerInputInfo::CalcNumNeuralNetInputs() const
+{
+	return 
+		BoneNameStrings.Num() * 4 +	// Four floats per bone.
+		CurveNameStrings.Num();		// One float per curve.
+}
diff --git a/Engine/Plugins/Experimental/MLDeformer/Source/MLDeformer/Private/MLDeformerInstance.cpp b/Engine/Plugins/Experimental/MLDeformer/Source/MLDeformer/Private/MLDeformerInstance.cpp
new file mode 100644
--- /dev/null
+++ b/Engine/Plugins/Experimental/MLDeformer/Source/MLDeformer/Private/MLDeformerInstance.cpp
@@ -0,0 +1,178 @@
+// Copyright Epic Games, Inc. All Rights Reserved.
+
+#include "MLDeformerInstance.h"
+#include "MLDeformer.h"
+#include "MLDeformerAsset.h"
+#include "MLDeformerInputInfo.h"
+#include "Components/SkeletalMeshComponent.h"
+#include "NeuralNetwork.h"
+#include "Animation/AnimInstance.h"
+#include "Components/SkeletalMeshComponent.h"
+
+void FMLDeformerInstance::Init(UMLDeformerAsset* Asset, USkeletalMeshComponent* SkelMeshComponent)
+{
+	DeformerAsset = Asset;
+	SkeletalMeshComponent = SkelMeshComponent;
+
+	if (SkelMeshComponent == nullptr || Asset == nullptr)
+	{
+		AssetBonesToSkelMeshMappings.Empty();
+		return;
+	}
+
+	USkeletalMesh* SkelMesh = SkelMeshComponent->SkeletalMesh;
+	if (SkelMesh)
+	{
+		// Init the bone mapping table.
+		const FMLDeformerInputInfo& InputInfo = DeformerAsset->GetInputInfo();
+		const int32 NumAssetBones = InputInfo.GetNumBones();
+		AssetBonesToSkelMeshMappings.Reset();
+		AssetBonesToSkelMeshMappings.AddUninitialized(NumAssetBones);
+
+		// For each bone in the deformer asset, find the matching bone index inside the skeletal mesh component.
+		for (int32 Index = 0; Index < NumAssetBones; ++Index)
+		{
+			const FName BoneName = InputInfo.GetBoneName(Index);
+			const int32 SkelMeshBoneIndex = SkeletalMeshComponent->GetBoneIndex(BoneName);
+			AssetBonesToSkelMeshMappings[Index] = SkelMeshBoneIndex;
+		}
+	}
+
+	// Perform a compatibility check.
+	bIsCompatible = SkelMesh && CheckCompatibility(SkelMeshComponent, true).IsEmpty();
+}
+
+FString FMLDeformerInstance::CheckCompatibility(USkeletalMeshComponent* InSkelMeshComponent, bool LogIssues) const
+{
+	// If we're not compatible, generate a compatibility string.
+	USkeletalMesh* SkelMesh = InSkelMeshComponent ? InSkelMeshComponent->SkeletalMesh : nullptr;
+	if (SkelMesh && !DeformerAsset->GetInputInfo().IsCompatible(SkelMesh))
+	{
+		const FString ErrorText = DeformerAsset->GetInputInfo().GenerateCompatibilityErrorString(SkelMesh);
+		check(!ErrorText.IsEmpty());
+		if (LogIssues)
+		{
+			UE_LOG(LogMLDeformer, Error, TEXT("ML Deformer '%s' isn't compatible with Skeletal Mesh '%s'.\nReason(s):\n%s"), 
+				*DeformerAsset->GetName(), 
+				*SkelMesh->GetName(), 
+				*ErrorText);
+		}
+		return ErrorText;
+	}
+
+	// No errors.
+	return FString();
+}
+
+int64 FMLDeformerInstance::SetBoneTransforms(float* OutputBuffer, int64 OutputBufferSize, int64 StartIndex)
+{
+	// Extract the component space bone transforms from the component.
+	// Write the output transforms into the BoneTransforms array.
+	BoneTransforms = SkeletalMeshComponent->GetBoneSpaceTransforms();
+
+	// Write the transforms into the output buffer.
+	const FMLDeformerInputInfo& InputInfo = DeformerAsset->GetInputInfo();
+	const int32 AssetNumBones = InputInfo.GetNumBones();
+	int64 Index = StartIndex;
+	check((Index + AssetNumBones * 4) <= OutputBufferSize); // Make sure we don't write past the OutputBuffer.
+	for (int32 BoneIndex = 0; BoneIndex < AssetNumBones; ++BoneIndex)
+	{
+		const int32 SkelMeshBoneIndex = AssetBonesToSkelMeshMappings[BoneIndex];
+		if (SkelMeshBoneIndex != INDEX_NONE)
+		{
+			const FQuat& Rotation = BoneTransforms[SkelMeshBoneIndex].GetRotation();
+			OutputBuffer[Index++] = Rotation.X;
+			OutputBuffer[Index++] = Rotation.Y;
+			OutputBuffer[Index++] = Rotation.Z;
+			OutputBuffer[Index++] = Rotation.W;
+		}
+		else
+		{
+			// Use an identity quaternion as input.
+			OutputBuffer[Index++] = 0.0f;
+			OutputBuffer[Index++] = 0.0f;
+			OutputBuffer[Index++] = 0.0f;
+			OutputBuffer[Index++] = 1.0f;
+		}
+	}
+
+	return Index;
+}
+
+int64 FMLDeformerInstance::SetCurveValues(float* OutputBuffer, int64 OutputBufferSize, int64 StartIndex)
+{
+	const FMLDeformerInputInfo& InputInfo = DeformerAsset->GetInputInfo();
+
+	// Write the weights into the output buffer.
+	int64 Index = StartIndex;
+	const int32 AssetNumCurves = InputInfo.GetNumCurves();
+	check((Index + AssetNumCurves) <= OutputBufferSize); // Make sure we don't write past the OutputBuffer.
+
+	// Write the curve weights to the output buffer.
+	UAnimInstance* AnimInstance = SkeletalMeshComponent->GetAnimInstance();
+	for (int32 CurveIndex = 0; CurveIndex < AssetNumCurves; ++CurveIndex)
+	{
+		const FName CurveName = InputInfo.GetCurveName(CurveIndex);
+		const float CurveValue = AnimInstance->GetCurveValue(CurveName);	// Outputs 0.0 when not found.
+		OutputBuffer[Index++] = CurveValue;
+	}
+
+	return Index;
+}
+
+void FMLDeformerInstance::SetNeuralNetworkInputValues(float* InputData, int64 NumInputFloats)
+{
+	check(SkeletalMeshComponent);
+
+	// Feed data to the network inputs.
+	int64 BufferOffset = 0;
+	BufferOffset = SetBoneTransforms(InputData, NumInputFloats, BufferOffset);
+	BufferOffset = SetCurveValues(InputData, NumInputFloats, BufferOffset);
+	check(BufferOffset == NumInputFloats);
+}
+
+void FMLDeformerInstance::Update()
+{
+	if (DeformerAsset == nullptr || SkeletalMeshComponent == nullptr || SkeletalMeshComponent->SkeletalMesh == nullptr || !bIsCompatible)
+	{
+		return;
+	}
+
+	// Get the network and make sure it's loaded.
+	UNeuralNetwork* NeuralNetwork = DeformerAsset->GetInferenceNeuralNetwork();
+	if (NeuralNetwork == nullptr || !NeuralNetwork->IsLoaded())
+	{
+		return;
+	}
+
+	// We only support GPU processing of the neural network at the moment.
+	check(NeuralNetwork->GetBackEnd() == ENeuralBackEnd::UEOnly);
+	check(NeuralNetwork->GetInputDeviceType() == ENeuralDeviceType::CPU);
+	check(NeuralNetwork->GetDeviceType() == ENeuralDeviceType::GPU);
+	check(NeuralNetwork->GetOutputDeviceType() == ENeuralDeviceType::GPU);
+
+	// If the neural network expects a different number of inputs, do nothing.
+	const int64 NumNeuralNetInputs = NeuralNetwork->GetInputTensor().Num();
+	const int64 NumDeformerAssetInputs = static_cast<int64>(DeformerAsset->GetInputInfo().CalcNumNeuralNetInputs());
+	ensureMsgf(
+		NumNeuralNetInputs == NumDeformerAssetInputs, 
+		TEXT("Neural network %s expects %d inputs, while the deformer asset expects to feed %d inputs."), 
+		*DeformerAsset->GetName(), 
+		NumNeuralNetInputs, 
+		NumDeformerAssetInputs);
+	if (NumNeuralNetInputs != NumDeformerAssetInputs)
+	{
+		return;
+	}
+
+	// Update and write the input values directly into the input tensor.
+	float* InputDataPointer = (float*)NeuralNetwork->GetInputDataPointerMutable();
+	SetNeuralNetworkInputValues(InputDataPointer, NumNeuralNetInputs);
+
+	// Output deltas will be available on GPU for the DeformerGraph via UMLDeformerDataProvider.
+	// So this does not actually modify our mesh directly. It just outputs the deltas, which are then used inside a deformer graph later on, accessible
+	// through the data provider class mentioned above.
+	// We pass ENeuralDeviceType::CPU as parameter because the inputs come from the CPU.
+	// We could later switch to Asynchronous processing if we want to run this in a background thread.
+	NeuralNetwork->Run();
+}
diff --git a/Engine/Plugins/Experimental/MLDeformer/Source/MLDeformer/Public/MLDeformer.h b/Engine/Plugins/Experimental/MLDeformer/Source/MLDeformer/Public/MLDeformer.h
new file mode 100644
--- /dev/null
+++ b/Engine/Plugins/Experimental/MLDeformer/Source/MLDeformer/Public/MLDeformer.h
@@ -0,0 +1,16 @@
+// Copyright Epic Games, Inc. All Rights Reserved.
+
+#pragma once
+
+#include "CoreMinimal.h"
+#include "Modules/ModuleManager.h"
+
+MLDEFORMER_API DECLARE_LOG_CATEGORY_EXTERN(LogMLDeformer, Log, All);
+
+class MLDEFORMER_API FMLDeformerModule : public IModuleInterface
+{
+public:
+	/** IModuleInterface implementation. */
+	virtual void StartupModule() override;
+	virtual void ShutdownModule() override;
+};
diff --git a/Engine/Plugins/Experimental/MLDeformer/Source/MLDeformer/Public/MLDeformerAsset.h b/Engine/Plugins/Experimental/MLDeformer/Source/MLDeformer/Public/MLDeformerAsset.h
new file mode 100644
--- /dev/null
+++ b/Engine/Plugins/Experimental/MLDeformer/Source/MLDeformer/Public/MLDeformerAsset.h
@@ -0,0 +1,274 @@
+// Copyright Epic Games, Inc. All Rights Reserved.
+
+#pragma once
+
+#include "CoreMinimal.h"
+#include "UObject/NoExportTypes.h"
+#include "MLDeformerVizSettings.h"
+#include "MLDeformerInputInfo.h"
+#include "RenderCommandFence.h"
+#include "UObject/Object.h"
+#include "RenderResource.h"
+#include "MLDeformerAsset.generated.h"
+
+// Forward declarations.
+class UNeuralNetwork;
+class USkeletalMesh;
+class UGeometryCache;
+class UAnimSequence;
+class UMLDeformerAsset;
+
+/** The activation function to use during the ML Deformer training process. */
+UENUM()
+enum class EActivationFunction : uint8
+{
+	Relu,
+	LRelu,
+	Tanh
+};
+
+/** The loss function to use during the ML Deformer training process. */
+UENUM()
+enum class ELossFunction : uint8
+{
+	L1,
+	MSE,
+	Shrinkage
+};
+
+/** The decay function to adapt the learning rate during the ML Deformer training process. */
+UENUM()
+enum class EDecayFunction : uint8
+{
+	Linear,
+	Multiplicative
+};
+
+/** The device where the training and testing will be running. */
+UENUM()
+enum class EDeviceType : uint8
+{
+	CPU,	/** Train using CPU. */
+	GPU		/** Train using GPU.*/
+};
+
+/** The training inputs. Specifies what data to include in training. */
+UENUM()
+enum class ETrainingInputs : uint8
+{
+	BonesAndCurves = 0, /** Include both bone rotations and curve values. */
+	BonesOnly,			/** Include only bone rotations. */
+	CurvesOnly			/** Include only curve values. */
+};
+
+class FVertexMapBuffer : public FVertexBufferWithSRV
+{
+public:
+	void Init(const TArray<int32>& InVertexMap)
+	{
+		VertexMap = InVertexMap;
+	}
+
+private:
+	void InitRHI() override;
+
+	TArray<int32> VertexMap;
+};
+
+/**
+ * The machine learning deformer asset.
+ * At runtime this contains only the data needed to run the neural network inference.
+ * In the editor it contains the skeletal mesh and geometry cache that are required to calculate vertex position deltas.
+ */
+UCLASS(BlueprintType, hidecategories=Object)
+class MLDEFORMER_API UMLDeformerAsset : public UObject
+{
+	GENERATED_BODY()
+
+public:
+	UMLDeformerAsset(const FObjectInitializer& ObjectInitializer);
+
+	//~UObject interface
+	virtual void Serialize(FArchive& Archive) override;
+	virtual void PostLoad() override;
+	virtual void BeginDestroy() override;
+	virtual bool IsReadyForFinishDestroy() override;
+#if WITH_EDITOR
+	virtual void PostEditChangeProperty(FPropertyChangedEvent& PropertyChangedEvent) override;
+#endif
+	//~End of UObject interface
+
+#if WITH_EDITOR
+	FText GetGeomCacheErrorText(UGeometryCache* InGeomCache) const;
+	FText GetAnimSequenceErrorText(UGeometryCache* InGeomCache, UAnimSequence* InAnimSequence) const;
+	FText GetVertexErrorText(USkeletalMesh* InSkelMesh, UGeometryCache* InGeomCache, const FText& SkelName, const FText& GeomCacheName) const;
+	FText GetBaseAssetChangedErrorText() const;
+	FText GetTargetAssetChangedErrorText() const;
+	FText GetInputsErrorText() const;
+	FText GetIncompatibleSkeletonErrorText(USkeletalMesh* InSkelMesh, UAnimSequence* InAnimSeq) const;
+
+	void UpdateCachedNumVertices();
+	bool IsCompatibleWithNeuralNet() const;
+	FMLDeformerInputInfo CreateInputInfo() const;
+#endif
+
+	void InitVertexMap();
+	void InitGPUData();
+
+	const TArray<int32>& GetVertexMap() const { return VertexMap; }
+	UNeuralNetwork* GetInferenceNeuralNetwork() const { return NeuralNetwork.Get(); }
+	const FVertexMapBuffer& GetVertexMapBuffer() const { return VertexMapBuffer; }
+	const FVector3f& GetVertexDeltaMean() const { return VertexDeltaMean; }
+	const FVector3f& GetVertexDeltaScale() const { return VertexDeltaScale; }
+	const FMLDeformerInputInfo& GetInputInfo() const { return InputInfo; }
+	FMLDeformerInputInfo& GetInputInfo() { return InputInfo; }
+	ETrainingInputs GetTrainingInputs() const { return TrainingInputs; }
+	ETrainingInputs GetTempTrainingInputs() const { return TempTrainingInputs; }
+	void SetTempTrainingInputs(ETrainingInputs Inputs) { TempTrainingInputs = Inputs; }
+	void SetInputInfo(const FMLDeformerInputInfo& Input) { InputInfo = Input; }
+
+#if WITH_EDITORONLY_DATA
+	const FTransform& GetAlignmentTransform() const { return AlignmentTransform; }
+
+	const USkeletalMesh* GetSkeletalMesh() const { return SkeletalMesh;  }
+	USkeletalMesh* GetSkeletalMesh() { return SkeletalMesh; }
+
+	const UGeometryCache* GetGeometryCache() const { return GeometryCache; }
+	UGeometryCache* GetGeometryCache() { return GeometryCache; }
+
+	const UAnimSequence* GetAnimSequence() const { return AnimSequence;  }
+	UAnimSequence* GetAnimSequence() { return AnimSequence; }
+
+	const UMLDeformerVizSettings* GetVizSettings() const { return VizSettings; }
+	UMLDeformerVizSettings* GetVizSettings() { return VizSettings; }
+
+	int32 GetNumFrames() const;
+	float GetDeltaCutoffLength() const { return DeltaCutoffLength; }
+#endif
+
+public:
+	/** The neural network to use for inference. */
+	UPROPERTY()
+	TObjectPtr<UNeuralNetwork> NeuralNetwork = nullptr;
+
+	/** This is an index per vertex in the mesh, indicating the imported vertex number from the source asset. */
+	UPROPERTY()
+	TArray<int32> VertexMap;
+
+	/** The vertex delta mean. This is passed to the deformer shader to reconstruct the correct output deltas. */
+	UPROPERTY()
+	FVector3f VertexDeltaMean = FVector3f::ZeroVector;
+
+	/** The vertex delta scale. This is passed to the deformer shader to rescale the output deltas. */
+	UPROPERTY()
+	FVector3f VertexDeltaScale = FVector3f::OneVector;
+
+	/** Cached number of skeltal mesh vertices. */
+	UPROPERTY()
+	int32 NumSkeletalMeshVerts = 0;
+
+	/** Cached number of geom cache vertices. */
+	UPROPERTY()
+	int32 NumGeomCacheVerts = 0;
+
+	/** 
+	 * The information about the neural network inputs. This contains things such as bone names, morph target names, etc. 
+	 * It also describes the order of inputs to the network.
+	 */
+	UPROPERTY()
+	FMLDeformerInputInfo InputInfo;
+
+	/** Describes what inputs we should train the neural network on. */
+	UPROPERTY(EditAnywhere, Category = "Inputs and Output")
+	ETrainingInputs TrainingInputs = ETrainingInputs::BonesAndCurves;
+	ETrainingInputs TempTrainingInputs = ETrainingInputs::BonesAndCurves;	// Work around a bug in Slate related to combobox changes.
+
+	/** GPU buffers for Vertex Map. */
+	FVertexMapBuffer VertexMapBuffer;
+
+	/** Fence used in render thread cleanup on destruction. */
+	FRenderCommandFence RenderResourceDestroyFence;
+
+#if WITH_EDITORONLY_DATA
+	/** The skeletal mesh that represents the linear skinned mesh. */
+	UPROPERTY(EditAnywhere, Category = "Base Mesh")
+	TObjectPtr<USkeletalMesh> SkeletalMesh = nullptr;
+
+	/** The geometry cache that represents the complex mesh deformations. */
+	UPROPERTY(EditAnywhere, Category = "Target Mesh")
+	TObjectPtr<UGeometryCache> GeometryCache = nullptr;
+
+	/** The animation sequence to apply to the base mesh. This has to match the animation of the target mesh's geometry cache. */
+	UPROPERTY(EditAnywhere, Category = "Base Mesh")
+	TObjectPtr<UAnimSequence> AnimSequence = nullptr;
+
+	/** The visualization settings. */
+	UPROPERTY()
+	TObjectPtr<UMLDeformerVizSettings> VizSettings = nullptr;
+
+	/** The number of hidden layers that the neural network model will have.\nHigher numbers will slow down performance but can deal with more complex deformations. */
+	UPROPERTY(EditAnywhere, BlueprintReadWrite, Category = "Training Settings", meta = (ClampMin = "1", ClampMax = "10"))
+	int32 NumHiddenLayers = 2;
+
+	/** The number of units/neurons per hidden layer. Higher numbers will slow down performance but allow for more complex mesh deformations. */
+	UPROPERTY(EditAnywhere, BlueprintReadWrite, Category = "Training Settings", meta = (ClampMin = "1"))
+	int32 NumNeuronsPerLayer = 256;
+
+	/** The number of frames per batch when training the model. */
+	UPROPERTY(EditAnywhere, BlueprintReadWrite, Category = "Training Settings", meta = (ClampMin = "1"))
+	int32 BatchSize = 32;
+
+	/** The number of epochs to process without any decay. */
+	UPROPERTY(EditAnywhere, BlueprintReadWrite, Category = "Training Settings", meta = (ClampMin = "1"))
+	int32 Epochs = 10;
+
+	/** The number of epochs to process that include a decay. */
+	UPROPERTY(EditAnywhere, BlueprintReadWrite, Category = "Training Settings", meta = (ClampMin = "1"))
+	int32 EpochsWithDecay = 10;
+
+	/** The learning rate used during the model training. */
+	UPROPERTY(EditAnywhere, BlueprintReadWrite, AdvancedDisplay, Category = "Training Settings", meta = (ClampMin = "0.000001", ClampMax = "1.0"))
+	float LearningRate = 0.0005f;
+
+	/** The decay function to adapt the learning rate. */
+	UPROPERTY(EditAnywhere, BlueprintReadWrite, AdvancedDisplay, Category = "Training Settings")
+	EDecayFunction DecayFunction = EDecayFunction::Linear;
+
+	/** The decay rate to apply to the learning rate once non-decay epochs have been reached. Higher values give less decay. */
+	UPROPERTY(EditAnywhere, BlueprintReadWrite, AdvancedDisplay, Category = "Training Settings", meta = (ClampMin = "0.1", ClampMax = "1.0"))
+	float DecayRate = 0.95f;
+
+	/** The activation function to use in the neural network. */
+	UPROPERTY(EditAnywhere, BlueprintReadWrite, AdvancedDisplay, Category = "Training Settings")
+	EActivationFunction ActivationFunction = EActivationFunction::Relu;
+
+	/** The loss function to use during model training. */
+	UPROPERTY(EditAnywhere, BlueprintReadWrite, AdvancedDisplay, Category = "Training Settings")
+	ELossFunction LossFunction = ELossFunction::L1;
+
+	/** Shrinkage speed. Only if the shrinkage loss is used. */
+	UPROPERTY(EditAnywhere, BlueprintReadWrite, AdvancedDisplay, Category = "Training Settings", meta = (ClampMin = "1.0", ClampMax = "1000.0"))
+	float ShrinkageSpeed = 10.0f;
+
+	/** Shrinkage threshold. Only if the shrinkage loss is used. */
+	UPROPERTY(EditAnywhere, BlueprintReadWrite, AdvancedDisplay, Category = "Training Settings", meta = (ClampMin = "0.0"))
+	float ShrinkageThreshold = 0.1f;
+
+	/** The percentage of noise to add. */
+	UPROPERTY(EditAnywhere, BlueprintReadWrite, AdvancedDisplay, Category = "Training Settings", meta = (ClampMin = "0.0", ClampMax = "100.0")) // A percentage.
+	float NoiseAmount = 0.5;
+
+	/** The loss function to use during model training. */
+	UPROPERTY(EditAnywhere, BlueprintReadWrite, AdvancedDisplay, Category = "Training Settings")
+	EDeviceType DeviceType = EDeviceType::GPU;
+
+	/** The transform that aligns the Geometry Cache to the SkeletalMesh. This will mostly apply some scale and a rotation, but no translation. */
+	UPROPERTY(EditAnywhere, BlueprintReadWrite, Category = "Target Mesh")
+	FTransform AlignmentTransform = FTransform::Identity;
+
+	/** Sometimes there can be some vertices that cause some issues that cause deltas to be very long. We can ignore these deltas by setting a cutoff value. 
+	  * Deltas that are longer than the cutoff value (in units), will be ignored and set to zero length. */
+	UPROPERTY(EditAnywhere, BlueprintReadWrite, Category = "Delta Generation", meta = (ClampMin = "0.01"))
+	float DeltaCutoffLength = 30.0f;
+#endif // WITH_EDITORONLY_DATA
+};
diff --git a/Engine/Plugins/Experimental/MLDeformer/Source/MLDeformer/Public/MLDeformerComponent.h b/Engine/Plugins/Experimental/MLDeformer/Source/MLDeformer/Public/MLDeformerComponent.h
new file mode 100644
--- /dev/null
+++ b/Engine/Plugins/Experimental/MLDeformer/Source/MLDeformer/Public/MLDeformerComponent.h
@@ -0,0 +1,75 @@
+// Copyright Epic Games, Inc. All Rights Reserved.
+
+#pragma once
+
+#include "CoreMinimal.h"
+#include "UObject/ObjectMacros.h"
+#include "InputCoreTypes.h"
+#include "Components/ActorComponent.h"
+#include "MLDeformer.h"
+#include "MLDeformerInstance.h"
+#include "MLDeformerComponent.generated.h"
+
+// Forward declarations.
+class UMLDeformerAsset;
+class USkeletalMeshComponent;
+
+/**
+ * The ML mesh deformer component.
+ * This works in combination with a MLDeformerAsset and SkeletalMeshComponent.
+ * The component will perform runtime inference using the network trained in the asset.
+ */
+UCLASS(Blueprintable, ClassGroup = Component, BlueprintType, meta = (BlueprintSpawnableComponent))
+class MLDEFORMER_API UMLDeformerComponent : public UActorComponent
+{
+	GENERATED_UCLASS_BODY()
+
+public:
+	UMLDeformerAsset* GetDeformerAsset() { return DeformerAsset; }
+	void SetDeformerAsset(UMLDeformerAsset* InDeformerAsset) { DeformerAsset = InDeformerAsset; }
+
+	//~UActorComponent interface
+	void Activate(bool bReset=false) override;
+	//End of ~UActorComponent interface
+
+	/**
+	 * Configure MLDeformerComponent with an MLRigDeformerAsset (network + parameters).
+	 * @param InDeformerAsset Pointer to MLDeformerAsset.
+	 */
+	UFUNCTION(BlueprintCallable, Category = "MLDeformer")
+	void SetupComponent(UMLDeformerAsset* InDeformerAsset, USkeletalMeshComponent* InSkelMeshComponent);
+
+	float GetVertexDeltaMultiplier() const { return VertexDeltaMultiplier; }
+	void SetVertexDeltaMultiplier(float ScaleFactor) { VertexDeltaMultiplier = ScaleFactor;  }
+
+	const FMLDeformerInstance& GetDeformerInstance() const { return DeformerInstance; }
+
+private:
+	/** The deformer asset to use. */
+	UPROPERTY(EditAnywhere, Category = "ML Deformer")
+	UMLDeformerAsset* DeformerAsset = nullptr;
+
+	/** The scale factor of the vertex deltas that are being applied. */
+	UPROPERTY(EditAnywhere, Category = "ML Deformer")
+	float VertexDeltaMultiplier = 1.0f;
+
+	/** 
+	 * The skeletal mesh component we want to grab the bone transforms etc from. 
+	 * This can be a nullptr. When it is a nullptr then it will internally try to find the first skeletal mesh component on the actor.
+	 * You can see this as an override. You can specify this override through the SetupComponent function.
+	 */
+	USkeletalMeshComponent* SkelMeshComponent = nullptr;
+
+	/** The mesh deformer instance. This is used to perform the runtime updates and run the inference. */
+	FMLDeformerInstance DeformerInstance;
+
+	// BEGIN AActorComponent interface
+	/**
+	 * Tick component, i.e., run MLDeformer inference.
+	 * @param DeltaTime Elapsed time since last tick
+	 * @param TickType Sort of tick
+	 * @param ThisTickFunction Tick function that triggers this tick.
+	 */
+	virtual void TickComponent(float DeltaTime, enum ELevelTick TickType, FActorComponentTickFunction* ThisTickFunction) override;
+	// END AActorComponent interface
+};
diff --git a/Engine/Plugins/Experimental/MLDeformer/Source/MLDeformer/Public/MLDeformerGraphDataInterface.h b/Engine/Plugins/Experimental/MLDeformer/Source/MLDeformer/Public/MLDeformerGraphDataInterface.h
new file mode 100644
--- /dev/null
+++ b/Engine/Plugins/Experimental/MLDeformer/Source/MLDeformer/Public/MLDeformerGraphDataInterface.h
@@ -0,0 +1,75 @@
+// Copyright Epic Games, Inc. All Rights Reserved.
+
+#pragma once
+
+#include "MLDeformer.h"
+#include "MLDeformerAsset.h"
+#include "OptimusComputeDataInterface.h"
+#include "ComputeFramework/ComputeDataProvider.h"
+#include "MLDeformerGraphDataInterface.generated.h"
+
+class FRDGBuffer;
+class FRDGBufferSRV;
+class FRHIShaderResourceView;
+class UMLDeformerComponent;
+class UNeuralNetwork;
+
+/** Compute Framework Data Interface for reading skeletal mesh. */
+UCLASS(Category = ComputeFramework)
+class MLDEFORMER_API UMLDeformerDataInterface : public UOptimusComputeDataInterface
+{
+	GENERATED_BODY()
+
+public:
+	//~ Begin UOptimusComputeDataInterface Interface
+	FString GetDisplayName() const override;
+	TArray<FOptimusCDIPinDefinition> GetPinDefinitions() const override;
+	//~ End UOptimusComputeDataInterface Interface
+
+	//~ Begin UComputeDataInterface Interface
+	void GetSupportedInputs(TArray<FShaderFunctionDefinition>& OutFunctions) const override;
+	void GetShaderParameters(TCHAR const* UID, FShaderParametersMetadataBuilder& OutBuilder) const override;
+	void GetHLSL(FString& OutHLSL) const override;
+	void GetSourceTypes(TArray<UClass*>& OutSourceTypes) const override;
+	UComputeDataProvider* CreateDataProvider(UObject* InOuter, TArrayView<TObjectPtr<UObject>> InSourceObjects) const override;
+	//~ End UComputeDataInterface Interface
+};
+
+
+/** Compute Framework Data Provider for reading skeletal mesh. */
+UCLASS(BlueprintType, editinlinenew, Category = ComputeFramework)
+class MLDEFORMER_API UMLDeformerDataProvider : public UComputeDataProvider
+{
+	GENERATED_BODY()
+
+public:
+	UPROPERTY(EditAnywhere, BlueprintReadWrite, Category = Binding)
+	TObjectPtr<UMLDeformerComponent> DeformerComponent = nullptr;
+
+	//~ Begin UComputeDataProvider Interface
+	bool IsValid() const override;
+	FComputeDataProviderRenderProxy* GetRenderProxy() override;
+	//~ End UComputeDataProvider Interface
+};
+
+
+class FMLDeformerDataProviderProxy : public FComputeDataProviderRenderProxy
+{
+public:
+	FMLDeformerDataProviderProxy(UMLDeformerComponent* DeformerComponent);
+
+	//~ Begin FComputeDataProviderRenderProxy Interface
+	void AllocateResources(FRDGBuilder& GraphBuilder) override;
+	void GetBindings(int32 InvocationIndex, TCHAR const* UID, FBindings& OutBindings) const override;
+	//~ End FComputeDataProviderRenderProxy Interface
+
+private:
+	UNeuralNetwork* NeuralNetwork = nullptr;
+	FRHIShaderResourceView* VertexMapBufferSRV = nullptr;
+	FRDGBuffer* Buffer = nullptr;
+	FRDGBufferSRV* BufferSRV = nullptr;
+	FVector VertexDeltaScale = FVector(1.0f, 1.0f, 1.0f);
+	FVector VertexDeltaMean = FVector(0.0f, 0.0f, 0.0f);
+	bool bCanRunNeuralNet = false;
+	float VertexDeltaMultiplier = 1.0f;
+};
diff --git a/Engine/Plugins/Experimental/MLDeformer/Source/MLDeformer/Public/MLDeformerInputInfo.h b/Engine/Plugins/Experimental/MLDeformer/Source/MLDeformer/Public/MLDeformerInputInfo.h
new file mode 100644
--- /dev/null
+++ b/Engine/Plugins/Experimental/MLDeformer/Source/MLDeformer/Public/MLDeformerInputInfo.h
@@ -0,0 +1,180 @@
+// Copyright Epic Games, Inc. All Rights Reserved.
+
+#pragma once
+
+#include "CoreMinimal.h"
+#include "MLDeformer.h"
+#include "MLDeformerInputInfo.generated.h"
+
+class USkeletalMesh;
+class USkeletalMeshComponent;
+class UAnimInstance;
+class UGeometryCache;
+
+struct MLDEFORMER_API FMLDeformerInputInfoInitSettings
+{
+	/** The skeletal mesh to initialize for. Cannot be nullptr. */
+	TObjectPtr<USkeletalMesh> SkeletalMesh = nullptr;
+
+	/** The target mesh. Cannot be a nullptr. */
+	TObjectPtr<UGeometryCache> TargetMesh = nullptr;
+
+	/** The list of bone names to include in training. When empty, all bones are included. These are ignored when bIncludeBones is set to false. */
+	TArray<FString> BoneNamesToInclude;
+
+	/* The list of curve names to include in training. When empty, all curves are included. These are ignored when bIncludeCurves is set to false. */
+	TArray<FString> CurveNamesToInclude;
+
+	/** Do we want to include bone rotations as inputs to the training? */
+	bool bIncludeBones = true;
+
+	/** Do we want to include curve values as inputs to the training? */
+	bool bIncludeCurves = true;
+};
+
+/**
+ * The neural network input information.
+ * This contains arrays of names for things such as bones and morph targets.
+ * Knowing what bones etc are used as inputs, and in what order, helps us feeding the data during inference.
+ * It can also help us detect issues, for example when the character we apply the deformer to is missing any of those bones.
+ */
+USTRUCT()
+struct MLDEFORMER_API FMLDeformerInputInfo
+{
+	GENERATED_BODY()
+
+public:
+	/**
+	 * Get the number of bones that we trained on.
+	 * @return The number of bones.
+	 */
+	int32 GetNumBones() const { return BoneNames.Num(); }
+
+	/**
+	 * Get the bone name as a string, for a given bone we included during training.
+	 * @param Index The bone index, which is a number in range of [0..GetNumBones()-1].
+	 * @result The name of the bone.
+	 */
+	const FString& GetBoneNameString(int32 Index) const { return BoneNameStrings[Index]; }
+
+	/**
+	 * Get the bone name as an FName, for a given bone we included during training.
+	 * @param Index The bone index, which is a number in range of [0..GetNumBones()-1].
+	 * @result The name of the bone.
+	 */
+	const FName GetBoneName(int32 Index) const { return BoneNames[Index]; }
+
+	/**
+	 * Get the number of curves that we trained on.
+	 * @return The number of curves.
+	 */
+	int32 GetNumCurves() const { return CurveNames.Num(); }
+
+	/**
+	 * Get the curve name as a string, for a given bone we included during training.
+	 * @param Index The curve index, which is a number in range of [0..GetNumCurves()-1].
+	 * @result The name of the curve.
+	 */
+	const FString& GetCurveNameString(int32 Index) const { return CurveNameStrings[Index]; }
+
+	/**
+	 * Get the curve name as an FName, for a given curve we included during training.
+	 * @param Index The curve index, which is a number in range of [0..GetNumCurves()-1].
+	 * @result The name of the curve.
+	 */
+	const FName GetCurveName(int32 Index) const { return CurveNames[Index]; }
+
+	/**
+	 * Initialize the inputs based on a skeletal mesh.
+	 * @param InitSettings The init settings.
+	 */
+	void Init(const FMLDeformerInputInfoInitSettings& InitSettings);
+
+	/**
+	 * Check whether the current inputs are compatible with a given skeletal mesh.
+	 * @param SkeletalMesh The skeletal mesh to check compatibility with. This may not be a nullptr.
+	 * @return Returns true when we can safely apply the ML Deformer to a character using this skeletal mesh, otherwise false is returned.
+	 * @note Use GenerateCompatibilityErrorString to get the error report.
+	 */
+	bool IsCompatible(USkeletalMesh* SkeletalMesh) const;
+
+	/**
+	 * Get the compatibility error report.
+	 * @param SkeletalMesh The skeletal mesh to check compatibility with.
+	 * @return Returns an empty string in case there are no compatibility issues, otherwise it contains a string that describes the issue(s).
+	 *         In case a nullptr is passed as SkeletalMesh parameter, an empty string is returned.
+	 */
+	FString GenerateCompatibilityErrorString(USkeletalMesh* SkeletalMesh) const;
+
+	/** 
+	 * Update the FName arrays based on the name string arrays.
+	 * This is automatically called on PostLoad of the UMLDeformerAsset.
+	 */
+	void UpdateFNames();
+
+	/** 
+	 * Extract the curve values for all curves we're interested in.
+	 * @param AnimInstance The anim instance we are grabbing the curve values from.
+	 * @param OutValues The array to write the values to. This array will be reset/resized by this method.
+	 */
+	void ExtractCurveValues(UAnimInstance* AnimInstance, TArray<float>& OutValues) const;
+
+	/**
+	 * Get the number of imported vertices in the base mesh.
+	 * @return The number of imported vertices in the base mesh.
+	 */
+	int32 GetNumBaseMeshVertices() const { return NumBaseMeshVertices; }
+
+	/**
+	 * Get the number of imported vertices in the target mesh.
+	 * @return The number of imported vertices in the target mesh.
+	 */
+	int32 GetNumTargetMeshVertices() const { return NumTargetMeshVertices; }
+
+	/**
+	 * Check whether we have any training inputs or not.
+	 * This happens when there are no bones or curves to use as inputs.
+	 * @return Returns true when there are no bones or curves specified as inputs.
+	 **/
+	bool IsEmpty() const { return (BoneNameStrings.IsEmpty() && CurveNameStrings.IsEmpty()); }
+
+	/**
+	 * Calculate how many inputs this input info generates for the neural network.
+	 * A single bone would take 4 inputs, while a curve takes one input.
+	 * @return THe number of input float values to the neural network.
+	 */
+	int32 CalcNumNeuralNetInputs() const;
+
+private:
+	/** 
+	 * The name of each bone. The inputs to the network are in the order of this array.
+	 * So if the array contains ["Root", "Child1", "Child2"] then the first bone transforms that we 
+	 * input to the neural network is the transform for "Root", followed by "Child1", followed by "Child2".
+	 */
+	UPROPERTY()
+	TArray<FString> BoneNameStrings;
+
+	/** The same as the BoneNames member, but stored as pre-created FName objects. These are not serialized. */
+	UPROPERTY(Transient)
+	TArray<FName> BoneNames;
+
+	/**
+	 * The name of each curve. The inputs to the network are in the order of this array.
+	 * So if the array contains ["Smile", "LeftEyeClosed", "RightEyeClosed"] then the first curve that we
+	 * input to the neural network is the one for "Smile", followed by "LeftEyeClosed", followed by "RightEyeClosed".
+	 */
+	UPROPERTY()
+	TArray<FString> CurveNameStrings;
+
+	/** The same as the CurveNames member, but stored as pre-created FName objects. These are not serialized. */
+	UPROPERTY(Transient)
+	TArray<FName> CurveNames;
+
+	/** Number of imported base mesh vertices, so not render vertices. */
+	UPROPERTY()
+	int32 NumBaseMeshVertices = 0;
+
+	/** Number of imported target mesh vertices, so not render vertices. */
+	UPROPERTY()
+	int32 NumTargetMeshVertices = 0;
+};
diff --git a/Engine/Plugins/Experimental/MLDeformer/Source/MLDeformer/Public/MLDeformerInstance.h b/Engine/Plugins/Experimental/MLDeformer/Source/MLDeformer/Public/MLDeformerInstance.h
new file mode 100644
--- /dev/null
+++ b/Engine/Plugins/Experimental/MLDeformer/Source/MLDeformer/Public/MLDeformerInstance.h
@@ -0,0 +1,103 @@
+// Copyright Epic Games, Inc. All Rights Reserved.
+
+#pragma once
+
+#include "CoreMinimal.h"
+#include "UObject/ObjectPtr.h"
+#include "Containers/Array.h"
+#include "Math/Transform.h"
+#include "Animation/AnimTypes.h"
+
+// Forward declarations.
+class UMLDeformerAsset;
+class USkeletalMeshComponent;
+
+/**
+ * This is an instance of the ML deformer, used during runtime.
+ * There will be one such instance per actor.
+ */
+class MLDEFORMER_API FMLDeformerInstance
+{
+public:
+	/**
+	 * Initialize this instance for a given deformer and skeletal mesh component.
+	 * @param Asset The ML Deformer asset.
+	 */
+	void Init(UMLDeformerAsset* Asset, USkeletalMeshComponent* SkelMeshComponent);
+
+	/** 
+	 * Update the deformer instance. 
+	 * This updates the neural net input values and runs the actual neural network inference to calculate the output deltas.
+	 */
+	void Update();
+
+	/**
+	 * Get the current component space bone transforms that we grabbed from the skeletal mesh component.
+	 * @return The cached bone transforms, in component space.
+	 */
+	const TArray<FTransform>& GetBoneTransforms() const { return BoneTransforms; }
+
+	/**
+	 * Is the deformer asset used compatible with the skeletal mesh component used during the Init call?
+	 * @return True if compatible, false if not.
+	 */
+	bool IsCompatible() const { return bIsCompatible; }
+
+	/**
+	 * Get the skeletal mesh component we're working with.
+	 * @return A pointer to the skeletal mesh component.
+	 */
+	USkeletalMeshComponent* GetSkeletalMeshComponent() const { return SkeletalMeshComponent; }
+
+private:
+	/**
+	 * Check whether the deformer is compatible with a given skeletal mesh component.
+	 * @param InSkelMeshComponent The skeletal mesh component to check compatibility with.
+	 * @param LogIssues Set to true to automatically log any compatibility errors.
+	 * @return Returns the error string. When the returned string is empty, there were no errors and thus
+	 *         the specified skeletal mesh component is compatible.
+	 */
+	FString CheckCompatibility(USkeletalMeshComponent* InSkelMeshComponent, bool LogIssues=false) const;
+
+	/**
+	 * Update the neural network input values directly inside its input tensor.
+	 * This will copy over the bone transforms (and possibly morph and curve weights) into this flat array.
+	 * @param InputData The buffer of floats to write the input values to.
+	 * @param NumInputFloats The number of floats of the InputData buffer. Do not write more than this number of floats to the InputData buffer.
+	 */
+	void SetNeuralNetworkInputValues(float* InputData, int64 NumInputFloats);
+
+	/**
+	 * Set the bone transformations inside a given output buffer, starting from a given StartIndex.
+	 * @param OutputBuffer The buffer we need to write the transformations to.
+	 * @param OutputBufferSize The number of floats inside the output buffer. Can be used to detect buffer overflows.
+	 * @param StartIndex The index where we want to start writing the transformations to.
+	 * @return The new start index into the output buffer, to be used when you write more data after this.
+	 */
+	int64 SetBoneTransforms(float* OutputBuffer, int64 OutputBufferSize, int64 StartIndex);
+
+	/**
+	 * Set the animation curve values inside a given output buffer, starting from a given StartIndex.
+	 * @param OutputBuffer The buffer we need to write the weights to.
+	 * @param OutputBufferSize The number of floats inside the output buffer. Can be used to detect buffer overflows.
+	 * @param StartIndex The index where we want to start writing the weights to.
+	 * @return The new start index into the output buffer, to be used when you write more data after this.
+	 */
+	int64 SetCurveValues(float* OutputBuffer, int64 OutputBufferSize, int64 StartIndex);
+
+private:
+	/** The ML Deformer asset we're an instance of. */
+	TObjectPtr<UMLDeformerAsset> DeformerAsset = nullptr;
+
+	/** The skeletal mesh component. */
+	TObjectPtr<USkeletalMeshComponent> SkeletalMeshComponent = nullptr;
+
+	/** The cached current component space bone transforms. */
+	TArray<FTransform> BoneTransforms;
+
+	/** Maps the ML deformer asset bone index to a skeletal mesh component bone index. */
+	TArray<int32> AssetBonesToSkelMeshMappings;
+
+	/** Are the deformer asset and the used skeletal mesh component compatible? */
+	bool bIsCompatible = false;
+};
diff --git a/Engine/Plugins/Experimental/MLDeformer/Source/MLDeformer/Public/MLDeformerVizSettings.h b/Engine/Plugins/Experimental/MLDeformer/Source/MLDeformer/Public/MLDeformerVizSettings.h
new file mode 100644
--- /dev/null
+++ b/Engine/Plugins/Experimental/MLDeformer/Source/MLDeformer/Public/MLDeformerVizSettings.h
@@ -0,0 +1,141 @@
+// Copyright Epic Games, Inc. All Rights Reserved.
+
+#pragma once
+
+#include "CoreMinimal.h"
+#include "UObject/NoExportTypes.h"
+#include "MLDeformerVizSettings.generated.h"
+
+// Forward declarations.
+class UAnimSequence;
+class UComputeGraph;
+class UGeometryCache;
+
+UENUM()
+enum class EMLDeformerVizMode : uint8
+{
+	TrainingData = 0,
+	TestData
+};
+
+/**
+ * The vizualization settings.
+ */
+UCLASS()
+class MLDEFORMER_API UMLDeformerVizSettings : public UObject
+{
+	GENERATED_BODY()
+
+public:
+	// UObject overloads.
+	virtual bool IsEditorOnly() const override { return true; }
+	// ~End of UObject interface
+
+#if WITH_EDITORONLY_DATA
+	const UAnimSequence* GetTestAnimSequence() const { return TestAnimSequence; }
+	UAnimSequence* GetTestAnimSequence() { return TestAnimSequence; }
+	UGeometryCache* GetGroundTruth() const { return GroundTruth; }
+	FVector GetMeshSpacingOffsetVector() const { return FVector(MeshSpacing, 0.0f, 0.0f); }
+	float GetMeshSpacing() const { return MeshSpacing; }
+	float GetLabelHeight() const { return LabelHeight; }
+	float GetAnimPlaySpeed() const { return AnimPlaySpeed; }
+	float GetVertexDeltaMultiplier() const { return VertexDeltaMultiplier; }
+	int32 GetFrameNumber() const { return FrameNumber; }
+	int32 GetVertexNumber() const { return VertexNumber; }
+	bool GetDrawVertex() const { return bDrawVertex; }
+	bool GetDrawLabels() const { return bDrawLabels; }
+	float GetLabelScale() const { return LabelScale; }
+	bool GetShowHeatMap() const { return bShowHeatMap; }
+	bool GetDrawVertexDeltas() const { return bDrawDeltas; }
+	bool GetDrawLinearSkinnedActor() const { return bDrawLinearSkinnedActor; }
+	bool GetDrawMLDeformedActor() const { return bDrawMLDeformedActor; }
+	bool GetDrawGroundTruthActor() const { return bDrawGroundTruthActor; }
+	bool GetXRayDeltas() const { return bXRayDeltas; }
+	UComputeGraph* GetDeformerGraph() const { return DeformerGraph; }
+	EMLDeformerVizMode GetVisualizationMode() const { return VisualizationMode; }
+	EMLDeformerVizMode GetTempVisualizationMode() const { return TempVisualizationMode; }
+	
+	void SetTempVisualizationMode(EMLDeformerVizMode Mode) { TempVisualizationMode = Mode; }
+	void SetDeformerGraph(UComputeGraph* InDeformerGraph) { DeformerGraph = InDeformerGraph; }
+#endif
+
+public:
+#if WITH_EDITORONLY_DATA
+	/** The data to visualize. */
+	UPROPERTY(EditAnywhere, Category = "Data");
+	EMLDeformerVizMode VisualizationMode = EMLDeformerVizMode::TrainingData;
+	EMLDeformerVizMode TempVisualizationMode = EMLDeformerVizMode::TrainingData;	// Workaround for a bug in UX where the combobox triggers a value changed when clicking it.
+
+	/** The animation sequence to play on the skeletal mesh. */
+	UPROPERTY(EditAnywhere, Category = "Testing")
+	TObjectPtr<UAnimSequence> TestAnimSequence = nullptr;
+
+	/** The deformer graph to use on the asset editor's deformed test actor. */
+	UPROPERTY(EditAnywhere, Category = "Testing")
+	TObjectPtr<UComputeGraph> DeformerGraph = nullptr;
+
+	/** The geometry cache that represents the ground truth of the test anim sequence. */
+	UPROPERTY(EditAnywhere, Category = "Testing")
+	TObjectPtr<UGeometryCache> GroundTruth = nullptr;
+
+	/** The scale factor of the ML deformer deltas being applied on top of the linear skinned results. */
+	UPROPERTY(EditAnywhere, Category = "Testing", meta = (ClampMin = "0.0", ClampMax = "5.0"))
+	float VertexDeltaMultiplier = 1.0f;
+
+	/** The play speed factor of the test anim sequence. */
+	UPROPERTY(EditAnywhere, Category = "Testing", meta = (ClampMin = "0.0", ClampMax = "2.0"))
+	float AnimPlaySpeed = 1.0f;
+
+	/** Show the heat map? This will visualize the active areas of the deformer. */
+	UPROPERTY(EditAnywhere, Category = "Testing")
+	bool bShowHeatMap = false;
+
+	/** Draw the linear skinned actor? */
+	UPROPERTY(EditAnywhere, Category = "Testing")
+	bool bDrawLinearSkinnedActor = true;
+
+	/** Draw the ML Deformed actor? */
+	UPROPERTY(EditAnywhere, Category = "Testing", DisplayName = "Draw ML Deformed Actor")
+	bool bDrawMLDeformedActor = true;
+
+	/** Draw the ground truth actor? */
+	UPROPERTY(EditAnywhere, Category = "Testing", meta = (EditCondition = "GroundTruth != nullptr"))
+	bool bDrawGroundTruthActor = true;
+
+	/** The frame number of the training data to visualize. */
+	UPROPERTY(EditAnywhere, Category = "Training Meshes", meta = (ClampMin = "0"))
+	uint32 FrameNumber = 0;
+
+	/** The vertex number to visualize. */
+	UPROPERTY(EditAnywhere, Category = "Training Meshes", meta = (ClampMin = "0"))
+	int32 VertexNumber = 0;
+
+	/** Draw the specified vertex? */
+	UPROPERTY(EditAnywhere, Category = "Training Meshes")
+	bool bDrawVertex = false;
+
+	/** Draw the vertex deltas? */
+	UPROPERTY(EditAnywhere, Category = "Training Meshes")
+	bool bDrawDeltas = true;
+
+	/** Enable this to draw the deltas in xray mode? */
+	UPROPERTY(EditAnywhere, Category = "Training Meshes", meta = (EditCondition = "bDrawDeltas"))
+	bool bXRayDeltas = true;
+
+	/** Draw the text labels above each actor? */
+	UPROPERTY(EditAnywhere, Category = "Shared Settings")
+	bool bDrawLabels = true;
+
+	/** The height in units to draw the labels at. */
+	UPROPERTY(EditAnywhere, Category = "Shared Settings", meta = (EditCondition = "bDrawLabels"))
+	float LabelHeight = 300.0f;
+
+	/** The scale of the label text. */
+	UPROPERTY(EditAnywhere, Category = "Shared Settings", meta = (ClampMin = "0.001", EditCondition = "bDrawLabels"))
+	float LabelScale = 1.0f;
+
+	/** The spacing between meshes. */
+	UPROPERTY(EditAnywhere, Category = "Shared Settings", meta = (ClampMin = "0"))
+	float MeshSpacing = 300.0f;
+#endif // WITH_EDITORONLY_DATA
+};
diff --git a/Engine/Plugins/Experimental/MLDeformer/Source/MLDeformerEditor/MLDeformerEditor.Build.cs b/Engine/Plugins/Experimental/MLDeformer/Source/MLDeformerEditor/MLDeformerEditor.Build.cs
new file mode 100644
--- /dev/null
+++ b/Engine/Plugins/Experimental/MLDeformer/Source/MLDeformerEditor/MLDeformerEditor.Build.cs
@@ -0,0 +1,57 @@
+// Copyright Epic Games, Inc. All Rights Reserved.
+
+using System.IO;
+
+namespace UnrealBuildTool.Rules
+{
+    public class MLDeformerEditor : ModuleRules
+    {
+        public MLDeformerEditor(ReadOnlyTargetRules Target) : base(Target)
+        {
+            PublicDependencyModuleNames.AddRange(
+                new string[]
+                {
+	                "Core",
+	                "AdvancedPreviewScene",
+                }
+            );
+
+			PrivateDependencyModuleNames.AddRange(
+                new string[]
+                {
+                    "InputCore",
+                    "EditorFramework",
+                    "UnrealEd",
+                    "ToolMenus",
+                    "CoreUObject",
+                    "Engine",
+                    "Slate",
+                    "SlateCore",
+                    "EditorWidgets",
+                    "Kismet",
+                    "KismetWidgets",
+                    "EditorStyle",
+                    "Persona",
+                    "SkeletonEditor",
+					"GeometryCache",
+
+					"PropertyEditor",
+					"BlueprintGraph",
+					"AnimGraph",
+					"AnimGraphRuntime",
+
+					"NeuralNetworkInference",
+
+					"MLDeformer",
+				}
+            );
+
+            PrivateIncludePathModuleNames.AddRange(
+                new string[]
+                {
+				}
+            );
+
+        }
+    }
+}
diff --git a/Engine/Plugins/Experimental/MLDeformer/Source/MLDeformerEditor/Private/MLDeformerApplicationMode.cpp b/Engine/Plugins/Experimental/MLDeformer/Source/MLDeformerEditor/Private/MLDeformerApplicationMode.cpp
new file mode 100644
--- /dev/null
+++ b/Engine/Plugins/Experimental/MLDeformer/Source/MLDeformerEditor/Private/MLDeformerApplicationMode.cpp
@@ -0,0 +1,110 @@
+// Copyright Epic Games, Inc. All Rights Reserved.
+
+#include "MLDeformerApplicationMode.h"
+#include "MLDeformerEditorToolkit.h"
+#include "MLDeformerVizSettingsTabSummoner.h"
+#include "MLDeformerTimelineTabSummoner.h"
+
+#include "IPersonaPreviewScene.h"
+#include "PersonaModule.h"
+#include "ISkeletonEditorModule.h"
+#include "Modules/ModuleManager.h"
+#include "PersonaTabs.h"
+
+#include "SSimpleTimeSlider.h"
+#include "Widgets/SBoxPanel.h"
+
+#define LOCTEXT_NAMESPACE "MLDeformerApplicationMode"
+
+FName FMLDeformerApplicationMode::ModeName("MLDeformerAssetEditMode");
+
+FMLDeformerApplicationMode::FMLDeformerApplicationMode(
+	TSharedRef<FWorkflowCentricApplication> InHostingApp,
+	TSharedRef<IPersonaPreviewScene> InPreviewScene)
+	: FApplicationMode(MLDeformerEditorModes::Editor)
+{
+	EditorToolkit = StaticCastSharedRef<FMLDeformerEditorToolkit>(InHostingApp);
+	TSharedRef<FMLDeformerEditorToolkit> MLDeformerEditor = StaticCastSharedRef<FMLDeformerEditorToolkit>(InHostingApp);
+
+	FPersonaViewportArgs ViewportArgs(InPreviewScene);
+	ViewportArgs.bAlwaysShowTransformToolbar = false;
+	ViewportArgs.bShowStats = false;
+	ViewportArgs.bShowTurnTable = false;
+	ViewportArgs.bShowTimeline = false;
+	ViewportArgs.bShowPhysicsMenu = false;
+	ViewportArgs.bShowLODMenu = true;
+	ViewportArgs.bShowPlaySpeedMenu = false;
+	ViewportArgs.ContextName = TEXT("MLDeformerEditor.Viewport");
+	ViewportArgs.OnViewportCreated = FOnViewportCreated::CreateSP(MLDeformerEditor, &FMLDeformerEditorToolkit::HandleViewportCreated);
+
+	// Register Persona tabs.
+	FPersonaModule& PersonaModule = FModuleManager::LoadModuleChecked<FPersonaModule>("Persona");
+	TabFactories.RegisterFactory(PersonaModule.CreatePersonaViewportTabFactory(InHostingApp, ViewportArgs));
+	TabFactories.RegisterFactory(PersonaModule.CreateDetailsTabFactory(InHostingApp, FOnDetailsCreated::CreateSP(&MLDeformerEditor.Get(), &FMLDeformerEditorToolkit::HandleDetailsCreated)));
+
+	// Register custom tabs.
+	TabFactories.RegisterFactory(MakeShared<FMLDeformerVizSettingsTabSummoner>(MLDeformerEditor));
+	TabFactories.RegisterFactory(MakeShared<FMLDeformerTimelineTabSummoner>(MLDeformerEditor));
+
+	// Create tab layout.
+	TabLayout = FTabManager::NewLayout("Standalone_MLDeformerEditor_Layout_v0.028")
+		->AddArea
+		(
+			FTabManager::NewPrimaryArea()
+			->SetOrientation(Orient_Vertical)
+			->Split
+			(
+				FTabManager::NewSplitter()
+				->SetSizeCoefficient(0.32f)
+				->SetOrientation(Orient_Horizontal)
+				->Split
+				(
+					FTabManager::NewStack()
+					->SetSizeCoefficient(0.3f)
+					->AddTab(FMLDeformerVizSettingsTabSummoner::TabID, ETabState::OpenedTab)
+				)
+				->Split
+				(
+					FTabManager::NewSplitter()
+					->SetSizeCoefficient(1.0f)
+					->SetOrientation(Orient_Vertical)
+					->Split
+					(
+						FTabManager::NewStack()
+						->SetSizeCoefficient(0.7f)
+						->SetHideTabWell(true)
+						->AddTab(FPersonaTabs::PreviewViewportID, ETabState::OpenedTab)
+					)
+					->Split
+					(
+						FTabManager::NewStack()
+						->SetSizeCoefficient(0.01f)
+						->SetHideTabWell(true)
+						->AddTab(FMLDeformerTimelineTabSummoner::TabID, ETabState::OpenedTab)
+					)
+				)
+				->Split
+				(
+					FTabManager::NewStack()
+					->SetSizeCoefficient(0.3f)
+					->AddTab(FPersonaTabs::DetailsID, ETabState::OpenedTab)
+					->SetForegroundTab(FPersonaTabs::DetailsID)
+				)
+			)
+		);
+
+	PersonaModule.OnRegisterTabs().Broadcast(TabFactories, InHostingApp);
+	LayoutExtender = MakeShared<FLayoutExtender>();
+	PersonaModule.OnRegisterLayoutExtensions().Broadcast(*LayoutExtender.Get());
+	TabLayout->ProcessExtensions(*LayoutExtender.Get());
+}
+
+void FMLDeformerApplicationMode::RegisterTabFactories(TSharedPtr<FTabManager> InTabManager)
+{
+	TSharedPtr<FMLDeformerEditorToolkit> Editor = EditorToolkit.Pin();
+	Editor->RegisterTabSpawners(InTabManager.ToSharedRef());
+	Editor->PushTabFactories(TabFactories);
+	FApplicationMode::RegisterTabFactories(InTabManager);
+}
+
+#undef LOCTEXT_NAMESPACE
diff --git a/Engine/Plugins/Experimental/MLDeformer/Source/MLDeformerEditor/Private/MLDeformerAssetActions.cpp b/Engine/Plugins/Experimental/MLDeformer/Source/MLDeformerEditor/Private/MLDeformerAssetActions.cpp
new file mode 100644
--- /dev/null
+++ b/Engine/Plugins/Experimental/MLDeformer/Source/MLDeformerEditor/Private/MLDeformerAssetActions.cpp
@@ -0,0 +1,35 @@
+// Copyright Epic Games, Inc. All Rights Reserved.
+
+#include "MLDeformerAssetActions.h"
+#include "MLDeformerAsset.h"
+#include "MLDeformerEditorToolkit.h"
+
+#define LOCTEXT_NAMESPACE "AssetTypeActions"
+
+UClass* FMLDeformerAssetActions::GetSupportedClass() const
+{
+	return UMLDeformerAsset::StaticClass();
+}
+
+void FMLDeformerAssetActions::GetActions(const TArray<UObject*>& InObjects, FToolMenuSection& Section)
+{
+	FAssetTypeActions_Base::GetActions(InObjects, Section);
+}
+
+
+void FMLDeformerAssetActions::OpenAssetEditor(const TArray<UObject*>& InObjects, TSharedPtr<IToolkitHost> EditWithinLevelEditor)
+{
+	const EToolkitMode::Type Mode = EditWithinLevelEditor.IsValid() ? EToolkitMode::WorldCentric : EToolkitMode::Standalone;
+
+	for (auto ObjIt = InObjects.CreateConstIterator(); ObjIt; ++ObjIt)
+	{
+		if (UMLDeformerAsset* Asset = Cast<UMLDeformerAsset>(*ObjIt))
+		{
+			TSharedRef<FMLDeformerEditorToolkit> NewEditor(new FMLDeformerEditorToolkit());
+			NewEditor->InitAssetEditor(Mode, EditWithinLevelEditor, Asset);
+		}
+	}
+}
+
+
+#undef LOCTEXT_NAMESPACE
diff --git a/Engine/Plugins/Experimental/MLDeformer/Source/MLDeformerEditor/Private/MLDeformerAssetActions.h b/Engine/Plugins/Experimental/MLDeformer/Source/MLDeformerEditor/Private/MLDeformerAssetActions.h
new file mode 100644
--- /dev/null
+++ b/Engine/Plugins/Experimental/MLDeformer/Source/MLDeformerEditor/Private/MLDeformerAssetActions.h
@@ -0,0 +1,20 @@
+// Copyright Epic Games, Inc. All Rights Reserved.
+
+#pragma once
+
+#include "CoreMinimal.h"
+//#include "Toolkits/IToolkitHost.h"
+#include "AssetTypeActions_Base.h"
+
+class FMLDeformerAssetActions : public FAssetTypeActions_Base
+{
+public:
+	// IAssetTypeActions Implementation
+	virtual FText GetName() const override { return NSLOCTEXT("AssetTypeActions", "AssetTypeActions_MLDeformer", "ML Deformer"); }
+	virtual void GetActions(const TArray<UObject*>& InObjects, struct FToolMenuSection& Section) override;
+	virtual FColor GetTypeColor() const override { return FColor(255, 255, 0); }
+	virtual UClass* GetSupportedClass() const override;
+	virtual void OpenAssetEditor(const TArray<UObject*>& InObjects, TSharedPtr<class IToolkitHost> EditWithinLevelEditor = TSharedPtr<IToolkitHost>()) override;
+	virtual uint32 GetCategories() override { return EAssetTypeCategories::Animation; }
+	//virtual UThumbnailInfo* GetThumbnailInfo(UObject* Asset) const override;
+};
diff --git a/Engine/Plugins/Experimental/MLDeformer/Source/MLDeformerEditor/Private/MLDeformerAssetDetails.cpp b/Engine/Plugins/Experimental/MLDeformer/Source/MLDeformerEditor/Private/MLDeformerAssetDetails.cpp
new file mode 100644
--- /dev/null
+++ b/Engine/Plugins/Experimental/MLDeformer/Source/MLDeformerEditor/Private/MLDeformerAssetDetails.cpp
@@ -0,0 +1,252 @@
+// Copyright Epic Games, Inc. All Rights Reserved.
+
+#include "MLDeformerAssetDetails.h"
+#include "MLDeformer.h"
+#include "MLDeformerAsset.h"
+#include "NeuralNetwork.h"
+
+#include "Animation/AnimSequence.h"
+#include "Engine/SkeletalMesh.h"
+#include "PropertyCustomizationHelpers.h"
+#include "AssetRegistry/AssetData.h"
+
+#include "DetailLayoutBuilder.h"
+#include "DetailCategoryBuilder.h"
+#include "DetailWidgetRow.h"
+#include "IDetailsView.h"
+
+#include "Widgets/DeclarativeSyntaxSupport.h"
+#include "Widgets/Layout/SBox.h"
+#include "DetailLayoutBuilder.h"
+
+#include "SWarningOrErrorBox.h"
+
+#define LOCTEXT_NAMESPACE "MLDeformerAssetDetails"
+
+
+TSharedRef<IDetailCustomization> FMLDeformerAssetDetails::MakeInstance()
+{
+	return MakeShareable(new FMLDeformerAssetDetails());
+}
+
+void FMLDeformerAssetDetails::CustomizeDetails(class IDetailLayoutBuilder& DetailBuilder)
+{
+	DetailLayoutBuilder = &DetailBuilder;
+
+	// Get the deformer asset.
+	TArray<TWeakObjectPtr<UObject>> Objects;
+	DetailBuilder.GetObjectsBeingCustomized(Objects);
+	UMLDeformerAsset* DeformerAsset = nullptr;
+	if (Objects.Num() == 1)
+	{
+		DeformerAsset = Cast<UMLDeformerAsset>(Objects[0]);
+	}
+
+	// Base mesh details.
+	IDetailCategoryBuilder& BaseMeshCategoryBuilder = DetailBuilder.EditCategory("Base Mesh", FText::GetEmpty(), ECategoryPriority::Important);
+	BaseMeshCategoryBuilder.AddProperty(GET_MEMBER_NAME_CHECKED(UMLDeformerAsset, SkeletalMesh));
+	if (DeformerAsset)
+	{
+		const FText ErrorText = DeformerAsset->GetVertexErrorText(DeformerAsset->SkeletalMesh, DeformerAsset->GeometryCache, FText::FromString("Base Mesh"), FText::FromString("Target Mesh"));
+		FDetailWidgetRow& ErrorRow = BaseMeshCategoryBuilder.AddCustomRow(FText::FromString("BaseMeshError"))
+			.Visibility(!ErrorText.IsEmpty() ? EVisibility::Visible : EVisibility::Collapsed)
+			.WholeRowContent()
+			[
+				SNew(SBox)
+				.Padding(FMargin(0.0f, 4.0f))
+				[
+					SNew(SWarningOrErrorBox)
+					.MessageStyle(EMessageStyle::Error)
+					.Message(ErrorText)
+				]
+			];
+
+		const FText ChangedErrorText = DeformerAsset->GetBaseAssetChangedErrorText();
+		FDetailWidgetRow& ChangedErrorRow = BaseMeshCategoryBuilder.AddCustomRow(FText::FromString("BaseMeshChangedError"))
+			.Visibility(!ChangedErrorText.IsEmpty() ? EVisibility::Visible : EVisibility::Collapsed)
+			.WholeRowContent()
+			[
+				SNew(SBox)
+				.Padding(FMargin(0.0f, 4.0f))
+				[
+					SNew(SWarningOrErrorBox)
+					.MessageStyle(EMessageStyle::Error)
+					.Message(ChangedErrorText)
+				]
+			];
+	}
+
+	// Animation sequence.
+	IDetailPropertyRow& AnimRow = BaseMeshCategoryBuilder.AddProperty(GET_MEMBER_NAME_CHECKED(UMLDeformerAsset, AnimSequence));
+	AnimRow.CustomWidget()
+	.NameContent()
+	[
+		AnimRow.GetPropertyHandle()->CreatePropertyNameWidget()
+	]
+	.ValueContent()
+	[
+		SNew(SObjectPropertyEntryBox)
+		.PropertyHandle(AnimRow.GetPropertyHandle())
+		.AllowedClass(UAnimSequence::StaticClass())
+		.ObjectPath(DeformerAsset ? DeformerAsset->GetAnimSequence()->GetPathName() : FString())
+		.ThumbnailPool(DetailBuilder.GetThumbnailPool())
+		.OnShouldFilterAsset(
+			this, 
+			&FMLDeformerAssetDetails::FilterAnimSequences, 
+			DeformerAsset->GetSkeletalMesh() ? DeformerAsset->GetSkeletalMesh()->GetSkeleton() : nullptr
+		)
+	];
+
+	if (DeformerAsset)
+	{
+		const FText WarningText = DeformerAsset->GetAnimSequenceErrorText(DeformerAsset->GetGeometryCache(), DeformerAsset->GetAnimSequence());
+		FDetailWidgetRow& WarningRow = BaseMeshCategoryBuilder.AddCustomRow(FText::FromString("AnimSeqWarning"))
+			.Visibility(!WarningText.IsEmpty() ? EVisibility::Visible : EVisibility::Collapsed)
+			.WholeRowContent()
+			[
+				SNew(SBox)
+				.Padding(FMargin(0.0f, 4.0f))
+				[
+					SNew(SWarningOrErrorBox)
+					.MessageStyle(EMessageStyle::Warning)
+					.Message(WarningText)
+				]
+			];
+
+		const FText AnimErrorText = DeformerAsset->GetIncompatibleSkeletonErrorText(DeformerAsset->GetSkeletalMesh(), DeformerAsset->AnimSequence);
+		FDetailWidgetRow& AnimErrorRow = BaseMeshCategoryBuilder.AddCustomRow(FText::FromString("AnimSkeletonMisMatchError"))
+			.Visibility(!AnimErrorText.IsEmpty() ? EVisibility::Visible : EVisibility::Collapsed)
+			.WholeRowContent()
+			[
+				SNew(SBox)
+				.Padding(FMargin(0.0f, 4.0f))
+				[
+					SNew(SWarningOrErrorBox)
+					.MessageStyle(EMessageStyle::Warning)
+					.Message(AnimErrorText)
+				]
+			];
+	}
+
+	// Target mesh details.
+	IDetailCategoryBuilder& TargetMeshCategoryBuilder = DetailBuilder.EditCategory("Target Mesh", FText::GetEmpty(), ECategoryPriority::Important);
+	TargetMeshCategoryBuilder.AddProperty(GET_MEMBER_NAME_CHECKED(UMLDeformerAsset, GeometryCache));
+	if (DeformerAsset)
+	{
+		const FText ErrorText = DeformerAsset->GetGeomCacheErrorText(DeformerAsset->GetGeometryCache());
+		FDetailWidgetRow& ErrorRow = TargetMeshCategoryBuilder.AddCustomRow(FText::FromString("TargetMeshError"))
+			.Visibility(!ErrorText.IsEmpty() ? EVisibility::Visible : EVisibility::Collapsed)
+			.WholeRowContent()
+			[
+				SNew(SBox)
+				.Padding(FMargin(0.0f, 4.0f))
+				[
+					SNew(SWarningOrErrorBox)
+					.MessageStyle(EMessageStyle::Error)
+					.Message(ErrorText)
+				]
+			];
+
+		const FText ChangedErrorText = DeformerAsset->GetTargetAssetChangedErrorText();
+		FDetailWidgetRow& ChangedErrorRow = BaseMeshCategoryBuilder.AddCustomRow(FText::FromString("TargetMeshChangedError"))
+			.Visibility(!ChangedErrorText.IsEmpty() ? EVisibility::Visible : EVisibility::Collapsed)
+			.WholeRowContent()
+			[
+				SNew(SBox)
+				.Padding(FMargin(0.0f, 4.0f))
+				[
+					SNew(SWarningOrErrorBox)
+					.MessageStyle(EMessageStyle::Error)
+					.Message(ChangedErrorText)
+				]
+			];
+	}
+	TargetMeshCategoryBuilder.AddProperty(GET_MEMBER_NAME_CHECKED(UMLDeformerAsset, AlignmentTransform));
+
+	// Input and output.
+	IDetailCategoryBuilder& InputOutputCategoryBuilder = DetailBuilder.EditCategory("Inputs and Output", FText::GetEmpty(), ECategoryPriority::Important);
+	InputOutputCategoryBuilder.AddProperty(GET_MEMBER_NAME_CHECKED(UMLDeformerAsset, TrainingInputs));
+	if (DeformerAsset)
+	{
+		const FText ErrorText = DeformerAsset->GetInputsErrorText();
+		FDetailWidgetRow& ErrorRow = InputOutputCategoryBuilder.AddCustomRow(FText::FromString("InputsError"))
+			.Visibility(!ErrorText.IsEmpty() ? EVisibility::Visible : EVisibility::Collapsed)
+			.WholeRowContent()
+			[
+				SNew(SBox)
+				.Padding(FMargin(0.0f, 4.0f))
+				[
+					SNew(SWarningOrErrorBox)
+					.MessageStyle(EMessageStyle::Error)
+					.Message(ErrorText)
+				]
+			];
+	}
+
+	InputOutputCategoryBuilder.AddProperty(GET_MEMBER_NAME_CHECKED(UMLDeformerAsset, DeltaCutoffLength));
+
+	// Training settings.
+	IDetailCategoryBuilder& SettingsCategoryBuilder = DetailBuilder.EditCategory("Training Settings", FText::GetEmpty(), ECategoryPriority::Important);
+	SettingsCategoryBuilder.AddProperty(GET_MEMBER_NAME_CHECKED(UMLDeformerAsset, NumHiddenLayers));
+	SettingsCategoryBuilder.AddProperty(GET_MEMBER_NAME_CHECKED(UMLDeformerAsset, NumNeuronsPerLayer));
+	SettingsCategoryBuilder.AddProperty(GET_MEMBER_NAME_CHECKED(UMLDeformerAsset, BatchSize));
+	SettingsCategoryBuilder.AddProperty(GET_MEMBER_NAME_CHECKED(UMLDeformerAsset, Epochs));
+	SettingsCategoryBuilder.AddProperty(GET_MEMBER_NAME_CHECKED(UMLDeformerAsset, EpochsWithDecay));
+
+	// Show a warning when no neural network has been set.
+	if (DeformerAsset)
+	{		
+		UNeuralNetwork* NeuralNetwork = DeformerAsset->GetInferenceNeuralNetwork();
+		FDetailWidgetRow& NeuralNetErrorRow = SettingsCategoryBuilder.AddCustomRow(FText::FromString("NeuralNetError"))
+			.Visibility((NeuralNetwork == nullptr) ? EVisibility::Visible : EVisibility::Collapsed)
+			.WholeRowContent()
+			[
+				SNew(SBox)
+				.Padding(FMargin(0.0f, 4.0f))
+				[
+					SNew(SWarningOrErrorBox)
+					.MessageStyle(EMessageStyle::Warning)
+					.Message(FText::FromString("Model still needs to be trained."))
+				]
+			];
+
+		// Check if our network is compatible with the skeletal mesh.
+		if (DeformerAsset->GetSkeletalMesh() && NeuralNetwork)
+		{
+			FDetailWidgetRow& NeuralNetIncompatibleErrorRow = SettingsCategoryBuilder.AddCustomRow(FText::FromString("NeuralNetIncompatibleError"))
+				.Visibility(!DeformerAsset->IsCompatibleWithNeuralNet() ? EVisibility::Visible : EVisibility::Collapsed)
+				.WholeRowContent()
+				[
+					SNew(SBox)
+					.Padding(FMargin(0.0f, 4.0f))
+					[
+						SNew(SWarningOrErrorBox)
+						.MessageStyle(EMessageStyle::Error)
+						.Message(FText::FromString("Trained neural network is incompatible with selected SkeletalMesh."))
+					]
+				];
+		}
+	}
+
+	SettingsCategoryBuilder.AddProperty(GET_MEMBER_NAME_CHECKED(UMLDeformerAsset, DecayFunction));
+	SettingsCategoryBuilder.AddProperty(GET_MEMBER_NAME_CHECKED(UMLDeformerAsset, DecayRate));
+	SettingsCategoryBuilder.AddProperty(GET_MEMBER_NAME_CHECKED(UMLDeformerAsset, LearningRate));
+	SettingsCategoryBuilder.AddProperty(GET_MEMBER_NAME_CHECKED(UMLDeformerAsset, ActivationFunction));
+	SettingsCategoryBuilder.AddProperty(GET_MEMBER_NAME_CHECKED(UMLDeformerAsset, LossFunction));
+	SettingsCategoryBuilder.AddProperty(GET_MEMBER_NAME_CHECKED(UMLDeformerAsset, ShrinkageSpeed));
+	SettingsCategoryBuilder.AddProperty(GET_MEMBER_NAME_CHECKED(UMLDeformerAsset, ShrinkageThreshold));
+	SettingsCategoryBuilder.AddProperty(GET_MEMBER_NAME_CHECKED(UMLDeformerAsset, NoiseAmount));
+	SettingsCategoryBuilder.AddProperty(GET_MEMBER_NAME_CHECKED(UMLDeformerAsset, DeviceType));
+}
+
+bool FMLDeformerAssetDetails::FilterAnimSequences(const FAssetData& AssetData, USkeleton* Skeleton)
+{
+	if (Skeleton && Skeleton->IsCompatibleSkeletonByAssetData(AssetData))
+	{
+		return false;
+	}
+
+	return true;
+}
+
+#undef LOCTEXT_NAMESPACE
diff --git a/Engine/Plugins/Experimental/MLDeformer/Source/MLDeformerEditor/Private/MLDeformerAssetDetails.h b/Engine/Plugins/Experimental/MLDeformer/Source/MLDeformerEditor/Private/MLDeformerAssetDetails.h
new file mode 100644
--- /dev/null
+++ b/Engine/Plugins/Experimental/MLDeformer/Source/MLDeformerEditor/Private/MLDeformerAssetDetails.h
@@ -0,0 +1,36 @@
+// Copyright Epic Games, Inc. All Rights Reserved.
+
+#pragma once
+
+#include "CoreMinimal.h"
+#include "IDetailCustomization.h"
+#include "AssetRegistry/AssetData.h"
+
+class FMLDeformerEditorData;
+class IDetailLayoutBuilder;
+class USkeleton;
+
+class FMLDeformerAssetDetails : public IDetailCustomization
+{
+public:
+	/** Makes a new instance of this detail layout class for a specific detail view requesting it. */
+	static TSharedRef<IDetailCustomization> MakeInstance();
+
+	/** ILayoutDetails interface. */
+	virtual void CustomizeDetails(class IDetailLayoutBuilder& DetailBuilder) override;
+
+	void SetTargetMeshErrorText(const FText& Text) { TargetMeshErrorText = Text; }
+
+private:
+	bool FilterAnimSequences(const FAssetData& AssetData, USkeleton* Skeleton);
+
+private:
+	/** Associated detail layout builder. */
+	IDetailLayoutBuilder* DetailLayoutBuilder = nullptr;
+
+	/** A pointer to the editor data. */
+	FMLDeformerEditorData* EditorData = nullptr;
+
+	/** The error text to show in the target mesh category. Empty if no error. */
+	FText TargetMeshErrorText;
+};
diff --git a/Engine/Plugins/Experimental/MLDeformer/Source/MLDeformerEditor/Private/MLDeformerAssetFactory.cpp b/Engine/Plugins/Experimental/MLDeformer/Source/MLDeformerEditor/Private/MLDeformerAssetFactory.cpp
new file mode 100644
--- /dev/null
+++ b/Engine/Plugins/Experimental/MLDeformer/Source/MLDeformerEditor/Private/MLDeformerAssetFactory.cpp
@@ -0,0 +1,52 @@
+// Copyright Epic Games, Inc. All Rights Reserved.
+
+
+#include "MLDeformerAssetFactory.h"
+#include "MLDeformerAsset.h"
+#include "AssetTypeCategories.h"
+
+#define LOCTEXT_NAMESPACE "MLDeformerAssetFactory"
+
+UMLDeformerFactory::UMLDeformerFactory()
+{
+	bCreateNew = true;
+	bEditAfterNew = true;
+	SupportedClass = UMLDeformerAsset::StaticClass();
+}
+
+UObject* UMLDeformerFactory::FactoryCreateNew(UClass* Class, UObject* InParent, FName Name, EObjectFlags Flags, UObject* Context, FFeedbackContext* Warn)
+{
+	return NewObject<UMLDeformerAsset>(InParent, Name, Flags | RF_Transactional);;
+}
+
+bool UMLDeformerFactory::ShouldShowInNewMenu() const
+{
+	return true;
+}
+
+bool UMLDeformerFactory::ConfigureProperties()
+{
+	return true;
+}
+
+FText UMLDeformerFactory::GetDisplayName() const
+{
+	return LOCTEXT("MLDeformerAsset_DisplayName", "ML Deformer");
+}
+
+uint32 UMLDeformerFactory::GetMenuCategories() const
+{
+	return EAssetTypeCategories::Animation;
+}
+
+FText UMLDeformerFactory::GetToolTip() const
+{
+	return LOCTEXT("MLDeformerAsset_Tooltip", "Machine learning based mesh deformer asset.");
+}
+
+FString UMLDeformerFactory::GetDefaultNewAssetName() const
+{
+	return FString(TEXT("NewMLDeformer"));
+}
+
+#undef LOCTEXT_NAMESPACE
diff --git a/Engine/Plugins/Experimental/MLDeformer/Source/MLDeformerEditor/Private/MLDeformerAssetFactory.h b/Engine/Plugins/Experimental/MLDeformer/Source/MLDeformerEditor/Private/MLDeformerAssetFactory.h
new file mode 100644
--- /dev/null
+++ b/Engine/Plugins/Experimental/MLDeformer/Source/MLDeformerEditor/Private/MLDeformerAssetFactory.h
@@ -0,0 +1,27 @@
+// Copyright Epic Games, Inc. All Rights Reserved.
+
+#pragma once
+
+#include "CoreMinimal.h"
+#include "Factories/Factory.h"
+#include "MLDeformerAssetFactory.generated.h"
+
+/**
+ * 
+ */
+UCLASS(hidecategories=Object)
+class UMLDeformerFactory : public UFactory
+{
+	GENERATED_BODY()
+
+	UMLDeformerFactory();
+
+	// UFactory Interface.
+	virtual FText GetDisplayName() const override;
+	virtual uint32 GetMenuCategories() const override;
+	virtual FText GetToolTip() const override;
+	virtual FString GetDefaultNewAssetName() const override;
+	virtual UObject* FactoryCreateNew(UClass* Class, UObject* InParent, FName Name, EObjectFlags Flags, UObject* Context, FFeedbackContext* Warn) override;
+	virtual bool ConfigureProperties() override;
+	virtual bool ShouldShowInNewMenu() const override;
+};
diff --git a/Engine/Plugins/Experimental/MLDeformer/Source/MLDeformerEditor/Private/MLDeformerEditor.cpp b/Engine/Plugins/Experimental/MLDeformer/Source/MLDeformerEditor/Private/MLDeformerEditor.cpp
new file mode 100644
--- /dev/null
+++ b/Engine/Plugins/Experimental/MLDeformer/Source/MLDeformerEditor/Private/MLDeformerEditor.cpp
@@ -0,0 +1,55 @@
+// Copyright Epic Games, Inc. All Rights Reserved.
+
+#include "MLDeformerEditor.h"
+#include "MLDeformerAssetActions.h"
+#include "MLDeformerAsset.h"
+#include "MLDeformerAssetDetails.h"
+#include "MLDeformerVizSettingsDetails.h"
+#include "MLDeformerEditorMode.h"
+#include "MLDeformer.h"
+
+#include "Modules/ModuleManager.h"
+#include "EditorModeRegistry.h"
+#include "PropertyEditorDelegates.h"
+
+#define LOCTEXT_NAMESPACE "MLDeformerEditorModule"
+
+IMPLEMENT_MODULE(FMLDeformerEditor, MLDeformerEditor)
+
+void FMLDeformerEditor::StartupModule()
+{
+	MLDeformerAssetActions = MakeShareable(new FMLDeformerAssetActions);
+	FModuleManager::LoadModuleChecked<FAssetToolsModule>("AssetTools").Get().RegisterAssetTypeActions(MLDeformerAssetActions.ToSharedRef());
+	FEditorModeRegistry::Get().RegisterMode<FMLDeformerEditorMode>(FMLDeformerEditorMode::ModeName, LOCTEXT("MLDeformerEditorMode", "MLDeformer"), FSlateIcon(), false);
+
+	// Register object detail customizations.
+	FPropertyEditorModule& PropertyModule = FModuleManager::GetModuleChecked<FPropertyEditorModule>("PropertyEditor");
+	PropertyModule.RegisterCustomClassLayout("MLDeformerAsset", FOnGetDetailCustomizationInstance::CreateStatic(&FMLDeformerAssetDetails::MakeInstance));
+	PropertyModule.RegisterCustomClassLayout("MLDeformerVizSettings", FOnGetDetailCustomizationInstance::CreateStatic(&FMLDeformerVizSettingsDetails::MakeInstance));
+	PropertyModule.NotifyCustomizationModuleChanged();
+}
+
+void FMLDeformerEditor::ShutdownModule()
+{
+	FEditorModeRegistry::Get().UnregisterMode(FMLDeformerEditorMode::ModeName);
+
+	if (MLDeformerAssetActions.IsValid())
+	{
+		if (FModuleManager::Get().IsModuleLoaded("AssetTools"))
+		{
+			FModuleManager::GetModuleChecked<FAssetToolsModule>("AssetTools").Get().UnregisterAssetTypeActions(MLDeformerAssetActions.ToSharedRef());
+		}
+		MLDeformerAssetActions.Reset();
+	}
+
+	// Unregister object detail customizations.
+	if (FModuleManager::Get().IsModuleLoaded(TEXT("PropertyEditor")))
+	{
+		FPropertyEditorModule& PropertyModule = FModuleManager::GetModuleChecked< FPropertyEditorModule >(TEXT("PropertyEditor"));
+		PropertyModule.UnregisterCustomClassLayout(TEXT("MLDeformerAsset"));
+		PropertyModule.UnregisterCustomClassLayout(TEXT("MLDeformerVizSettings"));
+		PropertyModule.NotifyCustomizationModuleChanged();
+	}
+}
+
+#undef LOCTEXT_NAMESPACE
diff --git a/Engine/Plugins/Experimental/MLDeformer/Source/MLDeformerEditor/Private/MLDeformerEditorData.cpp b/Engine/Plugins/Experimental/MLDeformer/Source/MLDeformerEditor/Private/MLDeformerEditorData.cpp
new file mode 100644
--- /dev/null
+++ b/Engine/Plugins/Experimental/MLDeformer/Source/MLDeformerEditor/Private/MLDeformerEditorData.cpp
@@ -0,0 +1,925 @@
+// Copyright Epic Games, Inc. All Rights Reserved.
+
+#include "MLDeformerEditorData.h"
+#include "MLDeformerAsset.h"
+#include "MLDeformerComponent.h"
+#include "MLDeformerInstance.h"
+#include "MLDeformerInputInfo.h"
+#include "MLDeformerEditorToolkit.h"
+#include "MLDeformer.h"
+
+#include "IPersonaToolkit.h"
+#include "IDetailsView.h"
+#include "Animation/DebugSkelMeshComponent.h"
+#include "Animation/AnimSequence.h"
+#include "Animation/AnimCurveTypes.h"
+#include "Animation/Skeleton.h"
+#include "Animation/SmartName.h"
+
+#include "Rendering/SkeletalMeshModel.h"
+#include "Materials/Material.h"
+
+#include "AnimPreviewInstance.h"
+
+#include "GeometryCache.h"
+#include "GeometryCacheComponent.h"
+#include "GeometryCacheMeshData.h"
+
+#include "Math/NumericLimits.h"
+#include "ComputeFramework/ComputeGraphComponent.h"
+#include "ComputeFramework/ComputeGraph.h"
+
+#include "SSimpleTimeSlider.h"
+
+#define LOCTEXT_NAMESPACE "MLDeformerEditorData"
+
+FMLDeformerEditorData::FMLDeformerEditorData()
+{
+	Actors.AddDefaulted(5);
+}
+
+int32 FMLDeformerEditorData::GetNumEditorActors() const
+{
+	return Actors.Num();
+}
+
+void FMLDeformerEditorData::SetEditorActor(EMLDeformerEditorActorIndex Index, const FMLDeformerEditorActor& Actor)
+{
+	Actors[static_cast<int32>(Index)] = Actor;
+}
+
+const FMLDeformerEditorActor& FMLDeformerEditorData::GetEditorActor(EMLDeformerEditorActorIndex Index) const
+{
+	return Actors[static_cast<int32>(Index)];
+}
+
+FMLDeformerEditorActor& FMLDeformerEditorData::GetEditorActor(EMLDeformerEditorActorIndex Index)
+{
+	return Actors[static_cast<int32>(Index)];
+}
+
+void FMLDeformerEditorData::SetPersonaToolkit(TSharedPtr<IPersonaToolkit> InPersonaToolkit)
+{
+	PersonaToolkit = InPersonaToolkit;
+}
+
+void FMLDeformerEditorData::SetDeformerAsset(UMLDeformerAsset* InAsset)
+{
+	MLDeformerAsset = InAsset;
+}
+
+void FMLDeformerEditorData::SetAnimInstance(UAnimInstance* InAnimInstance)
+{
+	AnimInstance = InAnimInstance;
+}
+
+void FMLDeformerEditorData::SetDetailsView(TSharedPtr<IDetailsView> InDetailsView)
+{
+	DetailsView = InDetailsView;
+}
+
+void FMLDeformerEditorData::SetVizSettingsDetailsView(TSharedPtr<IDetailsView> InDetailsView)
+{
+	VizSettingsDetailsView = InDetailsView;
+}
+
+void FMLDeformerEditorData::SetEditorToolkit(FMLDeformerEditorToolkit* InToolkit)
+{
+	EditorToolkit = InToolkit;
+}
+
+IPersonaToolkit* FMLDeformerEditorData::GetPersonaToolkit() const
+{ 
+	return PersonaToolkit.Get(); 
+}
+
+TSharedPtr<IPersonaToolkit> FMLDeformerEditorData::GetPersonaToolkitPointer() const
+{ 
+	return PersonaToolkit;
+}
+
+FMLDeformerEditorToolkit* FMLDeformerEditorData::GetEditorToolkit() const
+{
+	return EditorToolkit;
+}
+
+UMLDeformerAsset* FMLDeformerEditorData::GetDeformerAsset() const
+{ 
+	check(MLDeformerAsset.Get());
+	return MLDeformerAsset.Get();
+}
+
+TWeakObjectPtr<UMLDeformerAsset> FMLDeformerEditorData::GetDeformerAssetPointer() const
+{
+	return MLDeformerAsset;
+}
+
+UAnimInstance* FMLDeformerEditorData::GetAnimInstance() const
+{ 
+	return AnimInstance; 
+}
+
+IDetailsView* FMLDeformerEditorData::GetDetailsView() const
+{ 
+	return DetailsView.Get();
+}
+
+IDetailsView* FMLDeformerEditorData::GetVizSettingsDetailsView() const
+{
+	return VizSettingsDetailsView.Get();
+}
+
+bool FMLDeformerEditorData::IsTestActor(EMLDeformerEditorActorIndex Index) const
+{
+	return (Index == EMLDeformerEditorActorIndex::Test || Index == EMLDeformerEditorActorIndex::DeformedTest || Index == EMLDeformerEditorActorIndex::GroundTruth);
+}
+
+bool FMLDeformerEditorData::IsTrainingActor(EMLDeformerEditorActorIndex Index) const
+{
+	return (Index == EMLDeformerEditorActorIndex::Base || Index == EMLDeformerEditorActorIndex::Target);
+}
+
+float FMLDeformerEditorData::GetDuration() const
+{
+	USkeletalMeshComponent* SkeletalMeshComponent = GetEditorActor(EMLDeformerEditorActorIndex::Base).SkelMeshComponent;
+	UGeometryCacheComponent* GeometryCacheComponent = GetEditorActor(EMLDeformerEditorActorIndex::Target).GeomCacheComponent;
+
+	if (SkeletalMeshComponent && GeometryCacheComponent)
+	{
+		const float GeomCacheDuration = GeometryCacheComponent->GetDuration();
+		const float AnimSequenceDuration = MLDeformerAsset->GetAnimSequence()->GetPlayLength();
+		if (FMath::IsNearlyEqual(GeomCacheDuration, AnimSequenceDuration, 0.001f))
+		{
+			return GeomCacheDuration;
+		}
+		else
+		{
+			// TODO: Show some visual error, of mismatching durations.
+			//UE_LOG(LogMLDeformer, Warning, TEXT("%f, %f"), GeomCacheDuration, AnimSequenceDuration);
+		}
+	}
+
+	return 0.0f;
+}
+
+float FMLDeformerEditorData::GetTimeAtFrame(int32 FrameNumber) const
+{
+	UGeometryCacheComponent* GeometryCacheComponent = GetEditorActor(EMLDeformerEditorActorIndex::Target).GeomCacheComponent;
+	if (GeometryCacheComponent)
+	{
+		const int32 ClampedFrameNumber = FMath::Clamp(FrameNumber, 0, GeometryCacheComponent->GetNumberOfFrames() - 1);
+		return GeometryCacheComponent->GetTimeAtFrame(ClampedFrameNumber);
+	}
+
+	return 0.0f;
+}
+
+float FMLDeformerEditorData::GetSnappedFrameTime(float InTime) const
+{
+	UGeometryCacheComponent* GeometryCacheComponent = GetEditorActor(EMLDeformerEditorActorIndex::Target).GeomCacheComponent;
+	if (GeometryCacheComponent)
+	{
+		const int32 FrameIndex = GeometryCacheComponent->GetFrameAtTime(InTime);
+		return GeometryCacheComponent->GetTimeAtFrame(FrameIndex);
+	}
+
+	return InTime;
+}
+
+int32 FMLDeformerEditorData::GetNumImportedVertices() const
+{
+	UGeometryCacheComponent* GeomCacheComponent = GetEditorActor(EMLDeformerEditorActorIndex::Target).GeomCacheComponent;
+	check(GeomCacheComponent);
+
+	UGeometryCache* GeomCache = GeomCacheComponent->GetGeometryCache();
+	if (GeomCache == nullptr)
+	{
+		return 0;
+	}
+
+	TArray<FGeometryCacheMeshData> MeshData;
+	GeomCache->GetMeshDataAtTime(GeomCacheComponent->GetAnimationTime(), MeshData);
+
+	if (MeshData.Num() > 0)
+	{
+		const TArray<uint32>& ImportedVertexNumbers = MeshData[0].ImportedVertexNumbers;
+		if (ImportedVertexNumbers.Num() > 0)
+		{
+			// Find the maximum value.
+			int32 MaxIndex = -1;
+			for (int32 i = 0; i < ImportedVertexNumbers.Num(); ++i)
+			{
+				MaxIndex = FMath::Max(static_cast<int32>(ImportedVertexNumbers[i]), MaxIndex);
+			}
+			return (MaxIndex + 1);
+		}
+		return 0;
+	}
+	return 0;
+}
+
+void FMLDeformerEditorData::UpdateTestAnimPlaySpeed()
+{
+	UMLDeformerVizSettings* VizSettings = GetDeformerAsset()->GetVizSettings();
+	check(VizSettings);
+
+	const float Speed = VizSettings->GetAnimPlaySpeed();
+
+	USkeletalMeshComponent* SkeletalMeshComponent = GetEditorActor(EMLDeformerEditorActorIndex::Test).SkelMeshComponent;
+	SkeletalMeshComponent->SetPlayRate(Speed);
+
+	SkeletalMeshComponent = GetEditorActor(EMLDeformerEditorActorIndex::DeformedTest).SkelMeshComponent;
+	SkeletalMeshComponent->SetPlayRate(Speed);
+
+	UGeometryCacheComponent* GeomCacheComponent = GetEditorActor(EMLDeformerEditorActorIndex::GroundTruth).GeomCacheComponent;
+	GeomCacheComponent->SetPlaybackSpeed(Speed);
+}
+
+void FMLDeformerEditorData::UpdateDeformerGraph()
+{	
+	const FMLDeformerEditorActor& EditorActor = GetEditorActor(EMLDeformerEditorActorIndex::DeformedTest);
+	UComputeGraphComponent* ComputeComponent = EditorActor.Actor->FindComponentByClass<UComputeGraphComponent>();
+	if (ComputeComponent)
+	{
+		UMLDeformerVizSettings* VizSettings = MLDeformerAsset->GetVizSettings();
+		check(VizSettings);
+		ComputeComponent->ComputeGraph = MLDeformerAsset->GetInferenceNeuralNetwork() ? VizSettings->GetDeformerGraph() : nullptr;
+		ComputeComponent->CreateDataProviders(true);
+	}
+}
+
+void FMLDeformerEditorData::InitAssets()
+{
+	USkeletalMeshComponent* SkeletalMeshComponent = GetEditorActor(EMLDeformerEditorActorIndex::Base).SkelMeshComponent;
+	SkeletalMeshComponent->SetSkeletalMesh(MLDeformerAsset->GetSkeletalMesh());
+	if (PersonaToolkit)
+	{
+		PersonaToolkit->GetPreviewScene()->SetPreviewMesh(MLDeformerAsset->GetSkeletalMesh());
+	}
+	SkeletalMeshComponent->SetAnimationMode(EAnimationMode::AnimationSingleNode);
+	SkeletalMeshComponent->SetAnimation(MLDeformerAsset->GetAnimSequence());
+	SkeletalMeshComponent->SetPosition(0.0f);
+	SkeletalMeshComponent->SetPlayRate(0.0f);
+	SkeletalMeshComponent->Play(false);
+
+	UGeometryCacheComponent* GeometryCacheComponent = GetEditorActor(EMLDeformerEditorActorIndex::Target).GeomCacheComponent;
+	GeometryCacheComponent->SetGeometryCache(MLDeformerAsset->GetGeometryCache());
+	GeometryCacheComponent->ResetAnimationTime();
+	GeometryCacheComponent->SetLooping(false);
+	GeometryCacheComponent->SetManualTick(true);
+	GeometryCacheComponent->Play();
+
+	UMLDeformerVizSettings* VizSettings = MLDeformerAsset->GetVizSettings();
+	check(VizSettings);
+	const float TestAnimSpeed = VizSettings->GetAnimPlaySpeed();
+
+	SkeletalMeshComponent = GetEditorActor(EMLDeformerEditorActorIndex::Test).SkelMeshComponent;
+	SkeletalMeshComponent->SetSkeletalMesh(MLDeformerAsset->GetSkeletalMesh());
+	SkeletalMeshComponent->SetAnimationMode(EAnimationMode::AnimationSingleNode);
+	SkeletalMeshComponent->SetAnimation(VizSettings->GetTestAnimSequence());
+	SkeletalMeshComponent->SetPosition(0.0f);
+	SkeletalMeshComponent->SetPlayRate(TestAnimSpeed);
+	SkeletalMeshComponent->Play(true);
+
+	SkeletalMeshComponent = GetEditorActor(EMLDeformerEditorActorIndex::DeformedTest).SkelMeshComponent;
+	SkeletalMeshComponent->SetSkeletalMesh(MLDeformerAsset->GetSkeletalMesh());
+	SkeletalMeshComponent->SetAnimationMode(EAnimationMode::AnimationSingleNode);
+	SkeletalMeshComponent->SetAnimation(VizSettings->GetTestAnimSequence());
+	SkeletalMeshComponent->SetPosition(0.0f);
+	SkeletalMeshComponent->SetPlayRate(TestAnimSpeed);
+	SkeletalMeshComponent->Play(true);
+
+	GeometryCacheComponent = GetEditorActor(EMLDeformerEditorActorIndex::GroundTruth).GeomCacheComponent;
+	GeometryCacheComponent->SetGeometryCache(VizSettings->GetGroundTruth());
+	GeometryCacheComponent->ResetAnimationTime();
+	GeometryCacheComponent->SetLooping(true);
+	GeometryCacheComponent->SetManualTick(SkeletalMeshComponent->bPauseAnims);
+	GeometryCacheComponent->SetPlaybackSpeed(TestAnimSpeed);
+	GeometryCacheComponent->Play();
+
+	bIsVertexDeltaNormalized = false;
+	CurrentFrame = -1;
+
+	UpdateTimeSlider();
+	MLDeformerAsset->UpdateCachedNumVertices();
+	UpdateDeformerGraph();
+}
+
+void FMLDeformerEditorData::ExtractBoneTransforms(TArray<FMatrix44f>& OutBoneMatrices, TArray<FTransform>& OutBoneTransforms) const
+{
+	USkeletalMeshComponent* SkelComp = GetEditorActor(EMLDeformerEditorActorIndex::Base).SkelMeshComponent;
+	check(SkelComp);
+
+	SkelComp->CacheRefToLocalMatrices(OutBoneMatrices);
+	OutBoneTransforms = SkelComp->GetBoneSpaceTransforms();
+}
+
+void FMLDeformerEditorData::ExtractSkinnedPositions(int32 LODIndex, TArray<FMatrix44f>& InBoneMatrices, TArray<FVector3f>& TempPositions, TArray<FVector3f>& OutPositions, bool bImportedVertices) const
+{
+	OutPositions.Reset();
+	TempPositions.Reset();
+
+	UDebugSkelMeshComponent* SkelMeshComponent = GetEditorActor(EMLDeformerEditorActorIndex::Base).SkelMeshComponent;
+	if (SkelMeshComponent == nullptr)
+	{
+		return;
+	}
+
+	USkeletalMesh* Mesh = SkelMeshComponent->SkeletalMesh;
+	if (Mesh == nullptr)
+	{
+		return;
+	}
+
+	FSkeletalMeshLODRenderData& SkelMeshLODData = Mesh->GetResourceForRendering()->LODRenderData[LODIndex];
+	FSkinWeightVertexBuffer* SkinWeightBuffer = SkelMeshComponent->GetSkinWeightBuffer(LODIndex);
+	USkeletalMeshComponent::ComputeSkinnedPositions(SkelMeshComponent, TempPositions, InBoneMatrices, SkelMeshLODData, *SkinWeightBuffer);
+
+	if (bImportedVertices)
+	{
+		// Get the originally imported vertex numbers from the DCC.
+		const FSkeletalMeshModel* SkeletalMeshModel = Mesh->GetImportedModel();
+		const TArray<int32>& ImportedVertexNumbers = SkeletalMeshModel->LODModels[/*LODIndex=*/0].MeshToImportVertexMap;
+		if (ImportedVertexNumbers.Num() > 0)
+		{
+			const int32 MaxIndex = SkeletalMeshModel->LODModels[/*LODIndex=*/0].MaxImportVertex;
+
+			// Store the vertex positions for the original imported vertices (8 vertices for a cube).
+			OutPositions.AddZeroed(MaxIndex + 1);
+			for (int32 Index = 0; Index < TempPositions.Num(); ++Index)
+			{
+				const int32 ImportedVertex = ImportedVertexNumbers[Index];
+				OutPositions[ImportedVertex] = TempPositions[Index];
+			}
+		}
+	}
+	else
+	{
+		OutPositions = TempPositions;
+	}
+}
+
+void FMLDeformerEditorData::ExtractGeomCachePositions(int32 LODIndex, TArray<FVector3f>& TempPositions, TArray<FVector3f>& OutPositions) const
+{
+	OutPositions.Reset();
+	TempPositions.Reset();
+
+	UGeometryCacheComponent* GeomCacheComponent = GetEditorActor(EMLDeformerEditorActorIndex::Target).GeomCacheComponent;
+	UGeometryCache* GeomCache = GeomCacheComponent->GetGeometryCache();
+	if (GeomCache == nullptr)
+	{
+		return;
+	}
+
+	const float Time = GeomCacheComponent->GetAnimationTime();
+	TArray<FGeometryCacheMeshData> MeshData;
+	GeomCache->GetMeshDataAtTime(Time, MeshData);
+
+	if (MeshData.Num() > 0)
+	{
+		const TArray<uint32>& ImportedVertexNumbers = MeshData[0].ImportedVertexNumbers;
+		if (ImportedVertexNumbers.Num() > 0)
+		{
+			TempPositions = MeshData[0].Positions;
+
+			// Find the maximum value.
+			int32 MaxIndex = -1;
+			for (int32 Index = 0; Index < ImportedVertexNumbers.Num(); ++Index)
+			{
+				MaxIndex = FMath::Max(static_cast<int32>(ImportedVertexNumbers[Index]), MaxIndex);
+			}
+			check(MaxIndex > -1);
+
+			OutPositions.AddZeroed(MaxIndex + 1);
+
+			// Transform the positions with the alignment transform.
+			const FTransform& Transform = GetDeformerAsset()->AlignmentTransform;
+			for (int32 Index = 0; Index < TempPositions.Num(); ++Index)
+			{
+				const int32 ImportedVertex = static_cast<int32>(ImportedVertexNumbers[Index]);
+				const FVector Pos = TempPositions[Index];
+				OutPositions[ImportedVertex] = Transform.TransformPosition(Pos);
+			}
+		}
+	}
+}
+
+void FMLDeformerEditorData::SetAnimFrame(int32 FrameNumber)
+{
+	if (CurrentFrame == FrameNumber)
+	{
+		return;
+	}
+
+	// Calculate the time in the animation data to visualize.
+	const float TimeOffset = GetTimeAtFrame(FrameNumber);
+
+	USkeletalMeshComponent* SkeletalMeshComponent = GetEditorActor(EMLDeformerEditorActorIndex::Base).SkelMeshComponent;
+	if (SkeletalMeshComponent)
+	{
+		SkeletalMeshComponent->SetPosition(TimeOffset);
+		SkeletalMeshComponent->bPauseAnims = true;
+		SkeletalMeshComponent->RefreshBoneTransforms();
+	}
+
+	UGeometryCacheComponent* GeometryCacheComponent = GetEditorActor(EMLDeformerEditorActorIndex::Target).GeomCacheComponent;
+	if (GeometryCacheComponent && GeometryCacheComponent->GetGeometryCache())
+	{
+		GeometryCacheComponent->SetManualTick(true);
+		GeometryCacheComponent->TickAtThisTime(TimeOffset, false, false, false);
+	}
+
+	CurrentFrame = FrameNumber;
+}
+
+void FMLDeformerEditorData::UpdateDebugPointData()
+{
+	if (GetDeformerAsset()->GetVizSettings()->GetVisualizationMode() != EMLDeformerVizMode::TrainingData)
+	{
+		return;
+	}
+
+	const FMLDeformerEditorActor& BaseEditorActor = GetEditorActor(EMLDeformerEditorActorIndex::Base);
+	if (BaseEditorActor.SkelMeshComponent == nullptr || BaseEditorActor.SkelMeshComponent->SkeletalMesh == nullptr)
+	{
+		return;
+	}
+
+	// Extract the bone transforms and skinned vertex positions.
+	ExtractBoneTransforms(CurrentBoneMatrices, CurrentBoneTransforms);
+	ExtractSkinnedPositions(/*LODIndex=*/0, CurrentBoneMatrices, TempVectorBuffer, LinearSkinnedPositions);
+
+	const FMLDeformerEditorActor& TargetEditorActor = GetEditorActor(EMLDeformerEditorActorIndex::Target);
+	if (TargetEditorActor.GeomCacheComponent == nullptr || TargetEditorActor.GeomCacheComponent->GetGeometryCache() == nullptr)
+	{
+		return;
+	}
+
+	// Extract Geometry Cache positions.
+	ExtractGeomCachePositions(/*LODIndex=*/0, TempVectorBuffer, GeomCachePositions);
+}
+
+bool FMLDeformerEditorData::IsReadyForTraining() const
+{
+	UMLDeformerAsset* Asset = GetDeformerAsset();
+
+	// Make sure we have picked required assets.
+	if (Asset == nullptr ||
+		Asset->GetGeometryCache() == nullptr ||
+		GetDeformerAsset()->GetAnimSequence() == nullptr ||
+		GetDeformerAsset()->GetSkeletalMesh() == nullptr)
+	{
+		return false;
+	}
+
+	// Now make sure the assets are compatible.
+	if (!Asset->GetVertexErrorText(Asset->GetSkeletalMesh(), Asset->GetGeometryCache(), FText(), FText()).IsEmpty() ||
+		!Asset->GetGeomCacheErrorText(Asset->GetGeometryCache()).IsEmpty())
+	{
+		return false;
+	}
+
+	// Make sure we have inputs.
+	if (Asset->CreateInputInfo().IsEmpty())
+	{
+		return false;
+	}
+
+	return true;
+}
+
+bool FMLDeformerEditorData::GenerateTrainingData(uint32 LODIndex, uint32 FrameNumber, float DeltaCutoffLength, TArray<float>& OutDeltas, TArray<FTransform>& OutBoneTransforms, TArray<float>& OutCurveValues)
+{
+	OutDeltas.Reset();
+	OutBoneTransforms.Reset();
+	OutCurveValues.Reset();
+
+	if (!IsReadyForTraining())
+	{
+		return false;
+	}
+
+	// Calculate the number of frames to output.
+	UGeometryCacheComponent* GeometryCacheComponent = GetEditorActor(EMLDeformerEditorActorIndex::Target).GeomCacheComponent;
+	USkeletalMeshComponent* SkelMeshComponent = GetEditorActor(EMLDeformerEditorActorIndex::Base).SkelMeshComponent;
+	const uint32 AnimNumFrames = static_cast<uint32>(GeometryCacheComponent->GetNumberOfFrames());
+	const uint32 FrameIndex = FMath::Min(AnimNumFrames, FrameNumber);
+
+	// Calculate the vertex positions (and required transforms in order to perform skinning).
+	SetAnimFrame(static_cast<int32>(FrameIndex));
+	ExtractBoneTransforms(CurrentBoneMatrices, OutBoneTransforms);	// We just use our OutBoneTransforms array here, it will later be overwritten again.
+	ExtractSkinnedPositions(LODIndex, CurrentBoneMatrices, TempVectorBuffer, LinearSkinnedPositions);
+	ExtractGeomCachePositions(LODIndex, TempVectorBuffer, GeomCachePositions);
+
+	// Grab the actual required bone transforms
+	ExtractInputBoneTransforms(OutBoneTransforms);	// This internally resizes the output bone transforms array, so we overwrite existing contents from before.
+	ExtractInputCurveValues(OutCurveValues);
+
+	// Calculate the deltas.
+	if (LinearSkinnedPositions.Num() == GeomCachePositions.Num())
+	{
+		const int32 NumVertices = LinearSkinnedPositions.Num();
+		OutDeltas.AddUninitialized(3 * NumVertices);
+
+		for (int32 VertexIndex = 0; VertexIndex < NumVertices; ++VertexIndex)
+		{
+			const FVector Delta = GeomCachePositions[VertexIndex] - LinearSkinnedPositions[VertexIndex];
+			const int32 ArrayIndex = 3 * VertexIndex;
+			if (Delta.Length() < DeltaCutoffLength)
+			{
+				OutDeltas[ArrayIndex] = Delta.X;
+				OutDeltas[ArrayIndex + 1] = Delta.Y;
+				OutDeltas[ArrayIndex + 2] = Delta.Z;
+			}
+			else
+			{
+				OutDeltas[ArrayIndex] = 0.0f;
+				OutDeltas[ArrayIndex + 1] = 0.0f;
+				OutDeltas[ArrayIndex + 2] = 0.0f;
+			}
+		}
+	}
+
+	return true;
+}
+
+bool FMLDeformerEditorData::GenerateDeltas(uint32 LODIndex, uint32 FrameNumber, float DeltaCutoffLength, TArray<float>& OutDeltas)
+{
+	return GenerateTrainingData(LODIndex, FrameNumber, DeltaCutoffLength, OutDeltas, CurrentBoneTransforms, CurrentCurveWeights);
+}
+
+void FMLDeformerEditorData::UpdateVertexDeltaMeanAndScale(
+	const TArray<FVector3f>& InGeomCachePositions,
+	const TArray<FVector3f>& InLinearSkinnedPositions,
+	float DeltaCutoffLength,
+	FVector3f& InOutMeanVertexDelta,
+	FVector3f& InOutVertexDeltaScale,
+	float& InOutCount) const
+{
+	const int32 NumVertices = LinearSkinnedPositions.Num();
+	check(InGeomCachePositions.Num() == InLinearSkinnedPositions.Num());
+
+	FVector MeanDelta = FVector::ZeroVector;
+	FVector MinDelta(TNumericLimits<float>::Max());
+	FVector MaxDelta(-TNumericLimits<float>::Max());
+	float ValidVertexCount = 0.0f;
+	for (int32 VertexIndex = 0; VertexIndex < NumVertices; ++VertexIndex)
+	{
+		const FVector Delta = GeomCachePositions[VertexIndex] - LinearSkinnedPositions[VertexIndex];
+		if (Delta.Length() < DeltaCutoffLength)
+		{
+			MeanDelta += Delta;
+			ValidVertexCount += 1.0f;
+
+			// Get minimum dimensions.
+			MinDelta = MinDelta.ComponentMin(Delta);
+			// Get maximum dimensions.
+			MaxDelta = MaxDelta.ComponentMax(Delta);
+		}
+	}
+
+	if (ValidVertexCount > 0.0f)
+	{
+		MeanDelta /= ValidVertexCount;
+	}
+
+	// Update global mean.
+	InOutCount += 1.0f;
+	FVector TempDiff = MeanDelta - InOutMeanVertexDelta;
+	InOutMeanVertexDelta = InOutMeanVertexDelta + (TempDiff / InOutCount);
+
+	// Update global scale.
+	TempDiff = MaxDelta - MinDelta;
+	InOutVertexDeltaScale = InOutVertexDeltaScale.ComponentMax(TempDiff.GetAbs());
+}
+
+bool FMLDeformerEditorData::ComputeVertexDeltaStatistics(uint32 LODIndex, float DeltaCutoffLength)
+{
+	if (!IsReadyForTraining() || bIsVertexDeltaNormalized)
+	{
+		return false;
+	}
+
+	// Calculate the number of frames to output.
+	UGeometryCacheComponent* GeometryCacheComponent = GetEditorActor(EMLDeformerEditorActorIndex::Target).GeomCacheComponent;
+	const uint32 AnimNumFrames = static_cast<uint32>(GeometryCacheComponent->GetNumberOfFrames());
+
+	// Initialize mean vertex delta and vertex delta scale.
+	MLDeformerAsset->VertexDeltaMean = FVector3f::ZeroVector;
+	MLDeformerAsset->VertexDeltaScale = FVector3f::OneVector;
+	FVector3f TempVertexDeltaScale = FVector3f::ZeroVector;
+	float Count = 0.0f;
+	TArray<FTransform> BoneTransforms;
+	for (uint32 FrameIndex = 0; FrameIndex < AnimNumFrames; FrameIndex++)
+	{
+		// Set the components to the right time values of the given frame.
+		SetAnimFrame(static_cast<int32>(FrameIndex));
+
+		// Extract the bone transforms.
+		ExtractBoneTransforms(CurrentBoneMatrices, BoneTransforms);
+
+		// Extract the linear skinned positions.
+		ExtractSkinnedPositions(LODIndex, CurrentBoneMatrices, TempVectorBuffer, LinearSkinnedPositions);
+
+		// Extract the geometry cache positions.
+		ExtractGeomCachePositions(LODIndex, TempVectorBuffer, GeomCachePositions);
+
+		// Update mean vertex delta.
+		UpdateVertexDeltaMeanAndScale(GeomCachePositions, LinearSkinnedPositions, DeltaCutoffLength,
+			MLDeformerAsset->VertexDeltaMean, TempVertexDeltaScale, Count);
+	}
+
+	if (Count > 0.0f)
+	{
+		MLDeformerAsset->VertexDeltaScale *= TempVertexDeltaScale.GetMax();
+		bIsVertexDeltaNormalized = true;
+	}
+
+	return true;
+}
+
+FText FMLDeformerEditorData::GetOverlayText()
+{
+	if (GetDeformerAsset() == nullptr)
+	{
+		return FText::GetEmpty();
+	}
+
+	// Check compatibility.
+	const UMLDeformerComponent* DeformerComponent = GetEditorActor(EMLDeformerEditorActorIndex::DeformedTest).MLDeformerComponent;
+	if (DeformerComponent)
+	{
+		const FMLDeformerInstance& DeformerInstance = DeformerComponent->GetDeformerInstance();
+		if (DeformerInstance.GetSkeletalMeshComponent() && 
+			DeformerInstance.GetSkeletalMeshComponent()->SkeletalMesh &&
+			!DeformerComponent->GetDeformerInstance().IsCompatible() )
+		{
+			return FText::FromString(TEXT("Trained deformer asset is not compatible with skeletal mesh!"));
+		}
+	}
+
+	return FText::GetEmpty();
+}
+
+void FMLDeformerEditorData::CreateMaterials()
+{
+	// Try load the heat map material.
+	const FString HeatMapPath = TEXT("/MLDeformer/Materials/HeatMap.HeatMap");
+	UObject* Object = StaticLoadObject(UMaterial::StaticClass(), nullptr, *HeatMapPath);
+	HeatMapMaterial = Cast<UMaterial>(Object);
+	check(HeatMapMaterial);
+}
+
+void FMLDeformerEditorData::SetHeatMapMaterialEnabled(bool bEnabled)
+{
+	USkeletalMeshComponent* Component = GetEditorActor(EMLDeformerEditorActorIndex::DeformedTest).SkelMeshComponent;
+	check(Component);
+
+	if (bEnabled)
+	{
+		for (int32 Index = 0; Index < Component->GetNumMaterials(); ++Index)
+		{
+			Component->SetMaterial(Index, HeatMapMaterial);
+		}
+	}
+	else
+	{
+		Component->EmptyOverrideMaterials();
+	}
+}
+
+void FMLDeformerEditorData::ExtractInputCurveValues(TArray<float>& OutValues) const
+{
+	USkeletalMeshComponent* SkelMeshComponent = GetEditorActor(EMLDeformerEditorActorIndex::Base).SkelMeshComponent;
+	check(SkelMeshComponent);
+	const FMLDeformerInputInfo& InputInfo = GetDeformerAsset()->GetInputInfo();
+	InputInfo.ExtractCurveValues(SkelMeshComponent->GetAnimInstance(), OutValues);
+}
+
+void FMLDeformerEditorData::ExtractInputBoneTransforms(TArray<FTransform>& OutTransforms) const
+{
+	USkeletalMeshComponent* SkelMeshComponent = GetEditorActor(EMLDeformerEditorActorIndex::Base).SkelMeshComponent;
+	check(SkelMeshComponent);
+
+	// Get a copy of the bone transforms.
+	const TArray<FTransform> BoneTransforms = SkelMeshComponent->GetBoneSpaceTransforms();
+
+	const FMLDeformerInputInfo& InputInfo = GetDeformerAsset()->GetInputInfo();
+	const int32 NumBones = InputInfo.GetNumBones();
+	OutTransforms.Reset(NumBones);
+
+	if (NumBones == BoneTransforms.Num())
+	{
+		OutTransforms.AddUninitialized(NumBones);
+
+		for (int32 Index = 0; Index < NumBones; ++Index)
+		{
+			const FName BoneName = InputInfo.GetBoneName(Index);
+			const int32 SkelMeshBoneIndex = SkelMeshComponent->GetBoneIndex(BoneName);
+			OutTransforms[Index] = BoneTransforms[SkelMeshBoneIndex];
+		}
+	}
+}
+
+bool FMLDeformerEditorData::IsActorVisible(EMLDeformerEditorActorIndex ActorIndex) const
+{
+	const FMLDeformerEditorActor& EditorActor = GetEditorActor(ActorIndex);
+	if (EditorActor.SkelMeshComponent)
+	{
+		return EditorActor.SkelMeshComponent->IsVisible();
+	}
+
+	if (EditorActor.GeomCacheComponent)
+	{
+		return EditorActor.GeomCacheComponent->IsVisible();
+	}
+
+	return true;
+}
+
+void FMLDeformerEditorData::SetActorVisibility(EMLDeformerEditorActorIndex ActorIndex, bool bIsVisible)
+{
+	UDebugSkelMeshComponent* Component = GetEditorActor(ActorIndex).SkelMeshComponent;
+	if (Component && bIsVisible != Component->IsVisible())
+	{
+		Component->SetVisibility(bIsVisible, true);
+		UpdateDeformerGraph();
+	}
+
+	UGeometryCacheComponent* GeomComponent = GetEditorActor(ActorIndex).GeomCacheComponent;
+	if (GeomComponent && bIsVisible != GeomComponent->IsVisible())
+	{
+		GeomComponent->SetVisibility(bIsVisible, true);
+		UpdateDeformerGraph();
+	}
+}
+
+void FMLDeformerEditorData::SetTimeSlider(TSharedPtr<SSimpleTimeSlider> InTimeSlider)
+{
+	TimeSlider = InTimeSlider;
+}
+
+SSimpleTimeSlider* FMLDeformerEditorData::GetTimeSlider() const
+{
+	return TimeSlider.Get();
+}
+
+void FMLDeformerEditorData::UpdateTimeSlider()
+{
+	const UMLDeformerAsset* DeformerAsset = GetDeformerAsset();
+	const UMLDeformerVizSettings* VizSettings = DeformerAsset->GetVizSettings();
+	if (VizSettings->GetVisualizationMode() == EMLDeformerVizMode::TrainingData)
+	{
+		const UAnimSequence* AnimSeq = DeformerAsset->GetAnimSequence();
+		const double Duration = AnimSeq ? AnimSeq->GetPlayLength() : 0.0;
+		SetTimeSliderRange(0.0, Duration);
+	}
+	else
+	if (VizSettings->GetVisualizationMode() == EMLDeformerVizMode::TestData)
+	{
+		const UAnimSequence* AnimSeq = DeformerAsset->GetVizSettings()->GetTestAnimSequence();
+		const double Duration = AnimSeq ? AnimSeq->GetPlayLength() : 0.0;
+		SetTimeSliderRange(0.0, Duration);
+	}
+}
+
+void FMLDeformerEditorData::SetTimeSliderRange(double StartTime, double EndTime)
+{
+	SSimpleTimeSlider* Slider = GetTimeSlider();
+	if (Slider)
+	{
+		Slider->SetTimeRange(StartTime, EndTime);
+		Slider->SetClampRange(StartTime, EndTime);
+	}
+}
+
+void FMLDeformerEditorData::OnTimeSliderScrubPositionChanged(double NewScrubTime, bool bIsScrubbing)
+{
+	UMLDeformerAsset* DeformerAsset = GetDeformerAsset();
+	UMLDeformerVizSettings* VizSettings = DeformerAsset->GetVizSettings();
+	if (VizSettings->GetVisualizationMode() == EMLDeformerVizMode::TrainingData)
+	{
+		UGeometryCacheComponent* GeomCacheComponent = GetEditorActor(EMLDeformerEditorActorIndex::Target).GeomCacheComponent;
+		if (GeomCacheComponent)
+		{
+			const int32 FrameNumber = GeomCacheComponent->GetFrameAtTime(static_cast<float>(NewScrubTime));
+			VizSettings->FrameNumber = FrameNumber;
+		}
+	}
+	else
+	if (VizSettings->GetVisualizationMode() == EMLDeformerVizMode::TestData)
+	{
+		float PlayOffset = NewScrubTime;
+
+		// If we have a ground truth model, let's actually align to the exact geom cache frame to prevent interpolation.
+		UGeometryCacheComponent* GroundTruthGeomCacheComponent = GetEditorActor(EMLDeformerEditorActorIndex::GroundTruth).GeomCacheComponent;
+		if (GroundTruthGeomCacheComponent->GetGeometryCache())
+		{
+			const int32 FrameNumber = GroundTruthGeomCacheComponent->GetFrameAtTime(static_cast<float>(NewScrubTime));
+			PlayOffset = GroundTruthGeomCacheComponent->GetTimeAtFrame(FrameNumber);
+		}
+
+		UDebugSkelMeshComponent* SkelMeshComponent = GetEditorActor(EMLDeformerEditorActorIndex::Test).SkelMeshComponent;
+		SkelMeshComponent->SetPosition(PlayOffset);
+		SkelMeshComponent->bPauseAnims = true;
+		SkelMeshComponent->RefreshBoneTransforms();
+
+		UDebugSkelMeshComponent* DefSkelMeshComponent = GetEditorActor(EMLDeformerEditorActorIndex::DeformedTest).SkelMeshComponent;
+		DefSkelMeshComponent->SetPosition(PlayOffset);
+		DefSkelMeshComponent->bPauseAnims = true;
+		DefSkelMeshComponent->RefreshBoneTransforms();
+
+		GroundTruthGeomCacheComponent->SetManualTick(true);
+		GroundTruthGeomCacheComponent->TickAtThisTime(PlayOffset, false, false, false);
+	}	
+}
+
+void FMLDeformerEditorData::OnPlayButtonPressed()
+{
+	UMLDeformerAsset* DeformerAsset = GetDeformerAsset();
+	UMLDeformerVizSettings* VizSettings = DeformerAsset->GetVizSettings();
+	if (VizSettings->GetVisualizationMode() == EMLDeformerVizMode::TestData)
+	{
+		UDebugSkelMeshComponent* SkelMeshComponent = GetEditorActor(EMLDeformerEditorActorIndex::Test).SkelMeshComponent;
+		SkelMeshComponent->bPauseAnims = !SkelMeshComponent->bPauseAnims;
+		SkelMeshComponent->RefreshBoneTransforms();
+
+		UDebugSkelMeshComponent* DefSkelMeshComponent = GetEditorActor(EMLDeformerEditorActorIndex::DeformedTest).SkelMeshComponent;
+		DefSkelMeshComponent->bPauseAnims = !DefSkelMeshComponent->bPauseAnims;
+		DefSkelMeshComponent->RefreshBoneTransforms();
+
+		UGeometryCacheComponent* GroundTruthGeometryCacheComponent = GetEditorActor(EMLDeformerEditorActorIndex::GroundTruth).GeomCacheComponent;
+		if (DefSkelMeshComponent->bPauseAnims)
+		{
+			GroundTruthGeometryCacheComponent->SetManualTick(true);
+		}
+		else
+		{
+			GroundTruthGeometryCacheComponent->SetManualTick(false);
+		}
+	}
+}
+
+double FMLDeformerEditorData::CalcTimelinePosition() const
+{
+	UMLDeformerAsset* DeformerAsset = GetDeformerAsset();
+	UMLDeformerVizSettings* VizSettings = DeformerAsset->GetVizSettings();
+	UDebugSkelMeshComponent* SkelMeshComponent = nullptr;
+	if (VizSettings->GetVisualizationMode() == EMLDeformerVizMode::TestData)
+	{
+		SkelMeshComponent = GetEditorActor(EMLDeformerEditorActorIndex::Test).SkelMeshComponent;
+	}
+	else if (VizSettings->GetVisualizationMode() == EMLDeformerVizMode::TrainingData)
+	{
+		SkelMeshComponent = GetEditorActor(EMLDeformerEditorActorIndex::Base).SkelMeshComponent;
+	}
+
+	check(SkelMeshComponent);
+	return SkelMeshComponent->GetPosition();
+}
+
+bool FMLDeformerEditorData::IsPlayingAnim() const
+{
+	UMLDeformerAsset* DeformerAsset = GetDeformerAsset();
+	UMLDeformerVizSettings* VizSettings = DeformerAsset->GetVizSettings();
+	if (VizSettings->GetVisualizationMode() == EMLDeformerVizMode::TestData)
+	{
+		return !GetEditorActor(EMLDeformerEditorActorIndex::Test).SkelMeshComponent->bPauseAnims;
+	}
+	return false;
+}
+
+FString FMLDeformerEditorData::GetDefaultDeformerGraphAssetPath()
+{
+	return FString(TEXT("/MLDeformer/Deformers/DefaultMLDeformerGraph.MLDeformerGraph"));
+}
+
+UComputeGraph* FMLDeformerEditorData::LoadDefaultDeformerGraph()
+{
+	const FString GraphAssetPath = FMLDeformerEditorData::GetDefaultDeformerGraphAssetPath();
+	UObject* Object = StaticLoadObject(UComputeGraph::StaticClass(), nullptr, *GraphAssetPath);
+	UComputeGraph* DeformerGraph = Cast<UComputeGraph>(Object);
+	if (DeformerGraph == nullptr)
+	{
+		UE_LOG(LogMLDeformer, Warning, TEXT("Failed to load default ML Deformer compute graph from: %s"), *GraphAssetPath);
+	}
+	else
+	{
+		UE_LOG(LogMLDeformer, Verbose, TEXT("Loaded default ML Deformer compute graph from: %s"), *GraphAssetPath);
+	}
+
+	return DeformerGraph;
+}
+
+
+void FMLDeformerEditorData::SetDefaultDeformerGraphIfNeeded()
+{
+	// Initialize the asset on the default plugin deformer graph.
+	if (GetDeformerAsset()->GetVizSettings()->GetDeformerGraph() == nullptr)
+	{
+		UComputeGraph* DefaultGraph = FMLDeformerEditorData::LoadDefaultDeformerGraph();
+		GetDeformerAsset()->GetVizSettings()->SetDeformerGraph(DefaultGraph);
+	}
+}
+
+#undef LOCTEXT_NAMESPACE
diff --git a/Engine/Plugins/Experimental/MLDeformer/Source/MLDeformerEditor/Private/MLDeformerEditorMode.cpp b/Engine/Plugins/Experimental/MLDeformer/Source/MLDeformerEditor/Private/MLDeformerEditorMode.cpp
new file mode 100644
--- /dev/null
+++ b/Engine/Plugins/Experimental/MLDeformer/Source/MLDeformerEditor/Private/MLDeformerEditorMode.cpp
@@ -0,0 +1,276 @@
+// Copyright Epic Games, Inc. All Rights Reserved.
+
+#include "MLDeformerEditorMode.h"
+#include "MLDeformerEditorData.h"
+#include "MLDeformerAssetDetails.h"
+#include "MLDeformerComponent.h"
+#include "MLDeformerInstance.h"
+#include "MLDeformerEditorToolkit.h"
+#include "MLDeformerEditorStyle.h"
+
+#include "EditorViewportClient.h"
+#include "AssetEditorModeManager.h"
+#include "IPersonaPreviewScene.h"
+#include "Animation/DebugSkelMeshComponent.h"
+#include "Materials/Material.h"
+
+#include "Components/TextRenderComponent.h"
+#include "GeometryCacheComponent.h"
+
+#include "Rendering/SkeletalMeshModel.h"
+#include "Rendering/SkeletalMeshLODModel.h"
+
+#include "SSimpleTimeSlider.h"
+
+FName FMLDeformerEditorMode::ModeName("MLDeformerAssetEditMode");
+
+FMLDeformerEditorMode::FMLDeformerEditorMode()
+{
+}
+
+void FMLDeformerEditorMode::SetEditorData(TSharedPtr<FMLDeformerEditorData> InEditorData)
+{
+	EditorData = InEditorData;
+}
+
+void FMLDeformerEditorMode::EncapsulateBounds(const FMLDeformerEditorActor* Actor, FBox& Box) const
+{
+	check(Actor);
+	const USkeletalMeshComponent* SkelMeshComponent = Actor->SkelMeshComponent;
+	const UGeometryCacheComponent* GeomCacheComponent = Actor->GeomCacheComponent;
+
+	if (SkelMeshComponent && SkelMeshComponent->IsVisible())
+	{
+		Box += SkelMeshComponent->Bounds.GetBox();
+	}
+	else if (GeomCacheComponent && GeomCacheComponent->IsVisible())
+	{
+		Box += GeomCacheComponent->Bounds.GetBox();
+	}
+}
+
+bool FMLDeformerEditorMode::GetCameraTarget(FSphere& OutTarget) const
+{
+	auto Data = EditorData.Pin();
+
+	FBox Box;
+	Box.Init();
+	EncapsulateBounds(&Data->GetEditorActor(EMLDeformerEditorActorIndex::Base), Box);
+	EncapsulateBounds(&Data->GetEditorActor(EMLDeformerEditorActorIndex::Target), Box);
+	EncapsulateBounds(&Data->GetEditorActor(EMLDeformerEditorActorIndex::Test), Box);
+	EncapsulateBounds(&Data->GetEditorActor(EMLDeformerEditorActorIndex::DeformedTest), Box);
+	EncapsulateBounds(&Data->GetEditorActor(EMLDeformerEditorActorIndex::GroundTruth), Box);
+
+	if (Box.IsValid == 1)
+	{
+		OutTarget = FSphere(Box.GetCenter(), Box.GetExtent().X * 0.75f);
+		return true;
+	}
+
+	return false;
+}
+
+IPersonaPreviewScene& FMLDeformerEditorMode::GetAnimPreviewScene() const
+{
+	return *static_cast<IPersonaPreviewScene*>(static_cast<FAssetEditorModeManager*>(Owner)->GetPreviewScene());
+}
+
+void FMLDeformerEditorMode::GetOnScreenDebugInfo(TArray<FText>& OutDebugInfo) const
+{
+}
+
+void FMLDeformerEditorMode::Render(const FSceneView* View, FViewport* Viewport, FPrimitiveDrawInterface* PDI)
+{
+	FEdMode::Render(View, Viewport, PDI);
+
+	FMLDeformerEditorData* Data = EditorData.Pin().Get();
+	const UMLDeformerAsset* DeformerAsset = Data->GetDeformerAsset();
+	if (DeformerAsset == nullptr)
+	{
+		return;
+	}
+
+	const UMLDeformerVizSettings* VizSettings = DeformerAsset->GetVizSettings();
+	if (VizSettings->GetVisualizationMode() == EMLDeformerVizMode::TrainingData)
+	{
+		// Draw a single vertex.
+		if (VizSettings->GetDrawVertex())
+		{
+			const FLinearColor VertexColor = FMLDeformerEditorStyle::Get().GetColor("MLDeformer.Vertex.Color");
+			const float VertexSize = FMLDeformerEditorStyle::Get().GetFloat("MLDeformer.Vertex.Size");
+			const TArray<FVector3f>& SkinnedPositions = Data->LinearSkinnedPositions;
+			const FVector Offset = VizSettings->GetMeshSpacingOffsetVector();
+			const TArray<FVector3f>& GeomCachePositions = Data->GeomCachePositions;
+			const int32 VertexIndex = VizSettings->GetVertexNumber();
+			if (VertexIndex < SkinnedPositions.Num())
+			{
+				const FVector BasePos = SkinnedPositions[VertexIndex];
+				PDI->DrawPoint(BasePos, VertexColor, VertexSize, 10);
+			}
+
+			if (VertexIndex < GeomCachePositions.Num())
+			{
+				const FVector TargetPos = GeomCachePositions[VertexIndex] + Offset;
+				PDI->DrawPoint(TargetPos, VertexColor, VertexSize, 10);
+			}
+		}
+
+		// Draw deltas.
+		if (VizSettings->GetDrawVertexDeltas() && (Data->VertexDeltas.Num() / 3) == Data->LinearSkinnedPositions.Num())
+		{
+			const FLinearColor DeltasColor = FMLDeformerEditorStyle::Get().GetColor("MLDeformer.Deltas.Color");
+			const uint8 DepthGroup = VizSettings->GetXRayDeltas() ? 100 : 0;
+			const TArray<FVector3f>& SkinnedPositions = Data->LinearSkinnedPositions;
+			for (int32 Index = 0; Index < Data->LinearSkinnedPositions.Num(); ++Index)
+			{
+				const FVector Start = SkinnedPositions[Index];
+				FVector End;
+				int32 ArrayIndex = 3 * Index;
+				End.X = Start.X + Data->VertexDeltas[ArrayIndex];
+				End.Y = Start.Y + Data->VertexDeltas[++ArrayIndex];
+				End.Z = Start.Z + Data->VertexDeltas[++ArrayIndex];
+
+				PDI->DrawLine(Start, End, DeltasColor, DepthGroup);
+			}
+		}
+	}
+}
+
+bool FMLDeformerEditorMode::AllowWidgetMove()
+{
+	return false;
+}
+
+bool FMLDeformerEditorMode::ShouldDrawWidget() const
+{
+	return false;
+}
+
+bool FMLDeformerEditorMode::UsesTransformWidget() const
+{
+	return false;
+}
+
+bool FMLDeformerEditorMode::UsesTransformWidget(UE::Widget::EWidgetMode CheckMode) const
+{
+	return false;
+}
+
+void FMLDeformerEditorMode::UpdateLabels()
+{
+	FMLDeformerEditorData* Data = EditorData.Pin().Get();
+	UMLDeformerAsset* DeformerAsset = Data->GetDeformerAsset();
+	check(DeformerAsset);
+	
+	const UMLDeformerVizSettings* VizSettings = DeformerAsset->GetVizSettings();
+	const bool bDrawTrainingActors = (VizSettings->GetVisualizationMode() == EMLDeformerVizMode::TrainingData);
+	const bool bDrawTestActors = (VizSettings->GetVisualizationMode() == EMLDeformerVizMode::TestData);
+
+	for (int32 Index = 0; Index < Data->GetNumEditorActors(); ++Index)
+	{
+		const EMLDeformerEditorActorIndex ActorIndex = static_cast<EMLDeformerEditorActorIndex>(Index);
+		const FMLDeformerEditorActor& EditorActor = Data->GetEditorActor(ActorIndex);
+		UTextRenderComponent* LabelComponent = EditorActor.LabelComponent;
+		if (LabelComponent)
+		{
+			const AActor* Actor = EditorActor.Actor;
+			const FVector ActorLocation = Actor->GetActorLocation();
+			const FVector AlignmentOffset = (EditorActor.GeomCacheComponent != nullptr) ? DeformerAsset->GetAlignmentTransform().GetTranslation() : FVector::ZeroVector;
+			LabelComponent->SetRelativeLocation(ActorLocation + FVector(0.0f, 0.0f, VizSettings->GetLabelHeight()) - AlignmentOffset);
+			LabelComponent->SetRelativeRotation(FQuat(FVector(0.0f, 0.0f, 1.0f), FMath::DegreesToRadians(90.0f)));
+			LabelComponent->SetRelativeScale3D(FVector(VizSettings->GetLabelScale()));
+
+			// Update visibility.
+			if (VizSettings->bDrawLabels)
+			{
+				LabelComponent->SetVisibility(Data->IsActorVisible(ActorIndex));
+
+				// Handle ground truth, disable its label when no ground truth asset was selected.
+				if (ActorIndex == EMLDeformerEditorActorIndex::GroundTruth && EditorActor.GeomCacheComponent->GetGeometryCache() == nullptr)
+				{
+					LabelComponent->SetVisibility(false);
+				}
+			}
+			else
+			{
+				LabelComponent->SetVisibility(false);
+			}
+		}
+	}
+}
+
+void FMLDeformerEditorMode::UpdateActors()
+{
+	FMLDeformerEditorData* Data = EditorData.Pin().Get();
+
+	UMLDeformerAsset* DeformerAsset = Data->GetDeformerAsset();
+	const UMLDeformerVizSettings* VizSettings = DeformerAsset->GetVizSettings();
+	check(VizSettings);
+
+	// Update base actor's transform.
+	AActor* BaseActor = Data->GetEditorActor(EMLDeformerEditorActorIndex::Base).Actor;
+	BaseActor->SetActorTransform(FTransform::Identity);
+
+	// Update target actor's transform.
+	AActor* TargetActor = Data->GetEditorActor(EMLDeformerEditorActorIndex::Target).Actor;
+	FTransform TargetTransform = DeformerAsset->GetAlignmentTransform();
+	TargetTransform.AddToTranslation(VizSettings->GetMeshSpacingOffsetVector());
+	TargetActor->SetActorTransform(TargetTransform);
+
+	// Update test actor's transform.
+	AActor* TestActor = Data->GetEditorActor(EMLDeformerEditorActorIndex::Test).Actor;
+	TestActor->SetActorTransform(FTransform::Identity);
+
+	// Update deformed test actor's transform.
+	AActor* DeformedTestActor = Data->GetEditorActor(EMLDeformerEditorActorIndex::DeformedTest).Actor;
+	FTransform DeformedTestTransform = FTransform::Identity;
+	DeformedTestTransform.AddToTranslation(VizSettings->GetMeshSpacingOffsetVector());
+	DeformedTestActor->SetActorTransform(DeformedTestTransform);
+
+	// Update the ground truth actor's transform.
+	AActor* GroundTruthActor = Data->GetEditorActor(EMLDeformerEditorActorIndex::GroundTruth).Actor;
+	FTransform GroundTruthTransform = DeformerAsset->GetAlignmentTransform();
+	GroundTruthTransform.AddToTranslation(VizSettings->GetMeshSpacingOffsetVector() * 2.0f);
+	GroundTruthActor->SetActorTransform(GroundTruthTransform);
+
+	// Update the vertex delta multiplier.
+	UMLDeformerComponent* DeformerComponent = Data->GetEditorActor(EMLDeformerEditorActorIndex::DeformedTest).MLDeformerComponent;
+	check(DeformerComponent);
+	DeformerComponent->SetVertexDeltaMultiplier(VizSettings->GetVertexDeltaMultiplier());
+}
+
+void FMLDeformerEditorMode::Tick(FEditorViewportClient* ViewportClient, float DeltaTime)
+{
+	FMLDeformerEditorData* Data = EditorData.Pin().Get();
+
+	UMLDeformerAsset* DeformerAsset = Data->GetDeformerAsset();
+	const UMLDeformerVizSettings* VizSettings = DeformerAsset->GetVizSettings();
+
+	// Calculate the vertex positions when needed.
+	if ((VizSettings->GetVisualizationMode() == EMLDeformerVizMode::TrainingData) && VizSettings->GetDrawVertex())
+	{
+		Data->UpdateDebugPointData();
+	}
+
+	UpdateActors();
+	UpdateLabels();
+
+	// Calculate the training deltas when needed.
+	if ((VizSettings->GetVisualizationMode() == EMLDeformerVizMode::TrainingData) && VizSettings->GetDrawVertexDeltas())
+	{
+		Data->GenerateDeltas(0, VizSettings->FrameNumber, DeformerAsset->DeltaCutoffLength, Data->VertexDeltas);
+	}
+
+	// Force us back into the pose we want to see.
+	if (VizSettings->GetVisualizationMode() == EMLDeformerVizMode::TrainingData)
+	{
+		Data->SetAnimFrame(VizSettings->FrameNumber);
+	}
+
+	FEdMode::Tick(ViewportClient, DeltaTime);
+}
+
+void FMLDeformerEditorMode::DrawHUD(FEditorViewportClient* ViewportClient, FViewport* Viewport, const FSceneView* View, FCanvas* Canvas)
+{
+	FEdMode::DrawHUD(ViewportClient, Viewport, View, Canvas);
+}
diff --git a/Engine/Plugins/Experimental/MLDeformer/Source/MLDeformerEditor/Private/MLDeformerEditorStyle.cpp b/Engine/Plugins/Experimental/MLDeformer/Source/MLDeformerEditor/Private/MLDeformerEditorStyle.cpp
new file mode 100644
--- /dev/null
+++ b/Engine/Plugins/Experimental/MLDeformer/Source/MLDeformerEditor/Private/MLDeformerEditorStyle.cpp
@@ -0,0 +1,59 @@
+// Copyright Epic Games, Inc. All Rights Reserved.
+
+#include "MLDeformerEditorStyle.h"
+#include "EditorStyleSet.h"
+#include "Styling/SlateStyle.h"
+#include "Styling/SlateStyleRegistry.h"
+#include "Styling/SlateStyleMacros.h"
+#include "Styling/SlateTypes.h"
+#include "Math/Color.h"
+
+#define MLDEFORMER_IMAGE_BRUSH(RelativePath, ...) FSlateImageBrush(RootToContentDir(RelativePath, TEXT(".png")), __VA_ARGS__)
+
+FMLDeformerEditorStyle::FMLDeformerEditorStyle() : FSlateStyleSet("MLDeformerEditorStyle")
+{
+	SetContentRoot(FPaths::EnginePluginsDir() / TEXT("Experimental/MLDeformer/Resources"));
+	const FVector2D Icon16x16(16.0f, 16.0f);
+	const FVector2D Icon20x20(20.0f, 20.0f);
+	const FVector2D Icon40x40(40.0f, 40.0f);
+
+	// Icons.
+	Set("MLDeformer.Timeline.PlayIcon", new MLDEFORMER_IMAGE_BRUSH("Icons/Play_40x", Icon20x20));
+	Set("MLDeformer.Timeline.PauseIcon", new MLDEFORMER_IMAGE_BRUSH("Icons/Pause_40x", Icon20x20));
+	Set("MLDeformer.Timeline.TabIcon", new MLDEFORMER_IMAGE_BRUSH("Icons/TimelineTab_40x", Icon16x16));
+	Set("MLDeformer.VizSettings.TabIcon", new MLDEFORMER_IMAGE_BRUSH("Icons/VizSettingsTab_40x", Icon40x40));
+
+	// Colors and sizes.
+	Set("MLDeformer.BaseMesh.WireframeColor", FLinearColor(0.0f, 1.0f, 0.0f));
+	Set("MLDeformer.BaseMesh.LabelColor", FLinearColor(0.0f, 1.0f, 0.0f));
+
+	Set("MLDeformer.TargetMesh.WireframeColor", FLinearColor(0.0f, 0.0f, 1.0f));
+	Set("MLDeformer.TargetMesh.LabelColor", FLinearColor(0.0f, 0.0f, 1.0f));
+
+	Set("MLDeformer.MLDeformedMesh.WireframeColor", FLinearColor(1.0f, 0.0f, 0.0f));
+	Set("MLDeformer.MLDeformedMesh.LabelColor", FLinearColor(1.0f, 0.0f, 0.0f));
+
+	Set("MLDeformer.GroundTruth.WireframeColor", FLinearColor(0.0f, 0.0f, 1.0f));
+	Set("MLDeformer.GroundTruth.LabelColor", FLinearColor(0.0f, 0.0f, 1.0f));
+
+	Set("MLDeformer.Deltas.Color", FLinearColor(0.0f, 1.0f, 0.0f));
+	Set("MLDeformer.Vertex.Color", FLinearColor(1.0f, 0.0f, 1.0f));
+	Set("MLDeformer.Vertex.Size", 5.0f);
+
+	Set("MLDeformer.DefaultLabelScale", 0.65f);
+
+	FSlateStyleRegistry::RegisterSlateStyle(*this);
+}
+
+FMLDeformerEditorStyle::~FMLDeformerEditorStyle()
+{
+	FSlateStyleRegistry::UnRegisterSlateStyle(*this);
+}
+
+FMLDeformerEditorStyle& FMLDeformerEditorStyle::Get()
+{
+	static FMLDeformerEditorStyle Inst;
+	return Inst;
+}
+
+#undef MLDEFORMER_IMAGE_BRUSH
diff --git a/Engine/Plugins/Experimental/MLDeformer/Source/MLDeformerEditor/Private/MLDeformerEditorStyle.h b/Engine/Plugins/Experimental/MLDeformer/Source/MLDeformerEditor/Private/MLDeformerEditorStyle.h
new file mode 100644
--- /dev/null
+++ b/Engine/Plugins/Experimental/MLDeformer/Source/MLDeformerEditor/Private/MLDeformerEditorStyle.h
@@ -0,0 +1,13 @@
+// Copyright Epic Games, Inc. All Rights Reserved.
+
+#pragma once
+
+#include "Styling/SlateStyle.h"
+
+class FMLDeformerEditorStyle : public FSlateStyleSet
+{
+public:
+	FMLDeformerEditorStyle();
+	~FMLDeformerEditorStyle();
+	static FMLDeformerEditorStyle& Get();
+};
diff --git a/Engine/Plugins/Experimental/MLDeformer/Source/MLDeformerEditor/Private/MLDeformerEditorToolkit.cpp b/Engine/Plugins/Experimental/MLDeformer/Source/MLDeformerEditor/Private/MLDeformerEditorToolkit.cpp
new file mode 100644
--- /dev/null
+++ b/Engine/Plugins/Experimental/MLDeformer/Source/MLDeformerEditor/Private/MLDeformerEditorToolkit.cpp
@@ -0,0 +1,788 @@
+// Copyright Epic Games, Inc. All Rights Reserved.
+
+#include "MLDeformerEditorToolkit.h"
+
+#include "MLDeformer.h"
+#include "MLDeformerAsset.h"
+#include "MLDeformerApplicationMode.h"
+#include "MLDeformerEditorMode.h"
+#include "MLDeformerEditorData.h"
+#include "MLDeformerComponent.h"
+#include "MLDeformerInputInfo.h"
+#include "MLDeformerVizSettingsDetails.h"
+#include "MLDeformerEditorStyle.h"
+
+#include "NeuralNetwork.h"
+
+#include "GeometryCacheComponent.h"
+#include "GeometryCache.h"
+#include "GeometryCacheMeshData.h"
+
+#include "Rendering/SkeletalMeshModel.h"
+
+#include "AnimationEditorPreviewActor.h"
+#include "EditorModeManager.h"
+#include "GameFramework/WorldSettings.h"
+#include "Modules/ModuleManager.h"
+#include "PersonaModule.h"
+#include "IPersonaToolkit.h"
+#include "IAssetFamily.h"
+#include "ISkeletonEditorModule.h"
+#include "IPersonaViewport.h"
+#include "Preferences/PersonaOptions.h"
+#include "AnimCustomInstanceHelper.h"
+#include "Animation/DebugSkelMeshComponent.h"
+#include "Framework/MultiBox/MultiBoxBuilder.h"
+#include "AnimPreviewInstance.h"
+#include "MLDeformerPythonTrainingModel.h"
+#include "UObject/GCObjectScopeGuard.h"
+#include "UObject/Object.h"
+#include "Components/TextRenderComponent.h"
+#include "ComputeFramework/ComputeGraphComponent.h"
+
+#include "Widgets/SBoxPanel.h"
+#include "Framework/Notifications/NotificationManager.h"
+
+#include "Misc/MessageDialog.h"
+
+#define LOCTEXT_NAMESPACE "MLDeformerEditorToolkit"
+
+const FName MLDeformerEditorModes::Editor("MLDeformerEditorMode");
+const FName MLDeformerEditorAppName = FName(TEXT("MLDeformerEditorApp"));
+
+FMLDeformerEditorToolkit::FMLDeformerEditorToolkit()
+{
+	EditorData = MakeShared<FMLDeformerEditorData>();
+}
+
+FMLDeformerEditorToolkit::~FMLDeformerEditorToolkit()
+{
+	EditorData.Reset();
+}
+
+void FMLDeformerEditorToolkit::InitAssetEditor(
+	const EToolkitMode::Type Mode,
+	const TSharedPtr<IToolkitHost>& InitToolkitHost,
+	UMLDeformerAsset* DeformerAsset)
+{
+	EditorData->SetDeformerAsset(DeformerAsset);
+
+	FPersonaToolkitArgs PersonaToolkitArgs;
+	PersonaToolkitArgs.OnPreviewSceneCreated = FOnPreviewSceneCreated::FDelegate::CreateSP(this, &FMLDeformerEditorToolkit::HandlePreviewSceneCreated);
+	
+	FPersonaModule& PersonaModule = FModuleManager::LoadModuleChecked<FPersonaModule>("Persona");
+	auto PersonaToolkit = PersonaModule.CreatePersonaToolkit(DeformerAsset, PersonaToolkitArgs);
+	EditorData->SetPersonaToolkit(PersonaToolkit);
+	EditorData->SetEditorToolkit(this);
+
+	TSharedRef<IAssetFamily> AssetFamily = PersonaModule.CreatePersonaAssetFamily(DeformerAsset);
+	AssetFamily->RecordAssetOpened(FAssetData(DeformerAsset));
+
+	const bool bCreateDefaultStandaloneMenu = true;
+	const bool bCreateDefaultToolbar = true;
+	FAssetEditorToolkit::InitAssetEditor(
+		Mode,
+		InitToolkitHost,
+		MLDeformerEditorAppName,
+		FTabManager::FLayout::NullLayout,
+		bCreateDefaultStandaloneMenu,
+		bCreateDefaultToolbar,
+		DeformerAsset);
+
+	// Create and set the application mode.
+	FMLDeformerApplicationMode* ApplicationMode = new FMLDeformerApplicationMode(SharedThis(this), PersonaToolkit->GetPreviewScene());
+	AddApplicationMode(MLDeformerEditorModes::Editor, MakeShareable(ApplicationMode));
+	SetCurrentMode(MLDeformerEditorModes::Editor);
+
+	// Activate the editor mode.
+	GetEditorModeManager().SetDefaultMode(FMLDeformerEditorMode::ModeName);
+	GetEditorModeManager().ActivateMode(FMLDeformerEditorMode::ModeName);
+
+	FMLDeformerEditorMode* EditorMode = static_cast<FMLDeformerEditorMode*>(GetEditorModeManager().GetActiveMode(FMLDeformerEditorMode::ModeName));
+	EditorMode->SetEditorData(EditorData);
+
+	ExtendToolbar();
+	RegenerateMenusAndToolbars();
+
+	OnSwitchedVisualizationMode();
+	EditorData->UpdateTimeSlider();
+}
+
+TSharedRef<IPersonaToolkit> FMLDeformerEditorToolkit::GetPersonaToolkit() const 
+{ 
+	return EditorData->GetPersonaToolkitPointer().ToSharedRef();
+}
+
+void FMLDeformerEditorToolkit::RegisterTabSpawners(const TSharedRef<FTabManager>& InTabManager)
+{
+	WorkspaceMenuCategory = InTabManager->AddLocalWorkspaceMenuCategory(LOCTEXT("WorkspaceMenu_MLDeformerEditor", "ML Deformer Editor"));
+	auto WorkspaceMenuCategoryRef = WorkspaceMenuCategory.ToSharedRef();
+
+	FAssetEditorToolkit::RegisterTabSpawners(InTabManager);
+}
+
+void FMLDeformerEditorToolkit::UnregisterTabSpawners(const TSharedRef<FTabManager>& InTabManager)
+{
+	FAssetEditorToolkit::UnregisterTabSpawners(InTabManager);
+}
+
+void FMLDeformerEditorToolkit::ExtendToolbar()
+{
+	TSharedPtr<FExtender> ToolbarExtender = MakeShareable(new FExtender);
+
+	AddToolbarExtender(ToolbarExtender);
+
+	ToolbarExtender->AddToolBarExtension(
+		"Asset",
+		EExtensionHook::After,
+		GetToolkitCommands(),
+		FToolBarExtensionDelegate::CreateSP(this, &FMLDeformerEditorToolkit::FillToolbar)
+	);
+}
+
+void FMLDeformerEditorToolkit::FillToolbar(FToolBarBuilder& ToolbarBuilder)
+{
+	ToolbarBuilder.BeginSection("Training");
+	{
+		ToolbarBuilder.AddToolBarButton(
+			FUIAction(
+				FExecuteAction::CreateLambda([this]()
+				{
+					UMLDeformerAsset* DeformerAsset = EditorData->GetDeformerAsset();
+					DeformerAsset->SetInputInfo(DeformerAsset->CreateInputInfo());
+
+					// Ask if we want to retrain the network if we already have something trained.
+					if (DeformerAsset->GetInferenceNeuralNetwork() != nullptr)
+					{
+						const FText ConfirmTitle(LOCTEXT("RetrainConfirmationTitle", "Re-train the network?"));
+						const EAppReturnType::Type ConfirmReturnType = FMessageDialog::Open(
+							EAppMsgType::YesNo, 
+							FText(LOCTEXT("RetrainConfirmationMessage", "This asset already has been trained.\n\nAre you sure you would like to re-train the network with your current settings?")),
+							&ConfirmTitle);
+
+						if (ConfirmReturnType == EAppReturnType::No || ConfirmReturnType == EAppReturnType::Cancel)
+						{
+							return;
+						}
+					}
+
+					UMLDeformerPythonTrainingModel* MLDeformerModel = UMLDeformerPythonTrainingModel::Get();
+            		FGCObjectScopeGuard ModelGuard(MLDeformerModel);
+					if (MLDeformerModel)
+					{
+						ShowNotification(LOCTEXT("StartTraining", "Starting training process"), SNotificationItem::ECompletionState::CS_Pending, true);
+
+						// Initialize the training inputs.
+						DeformerAsset->SetInputInfo(DeformerAsset->CreateInputInfo());
+
+						// Make sure we have something to train on.
+						// If this triggers, the train button most likely was enabled while it shouldn't be.
+						check(!DeformerAsset->GetInputInfo().IsEmpty());
+
+						// Perform training and load the resulting model.
+						MLDeformerModel->SetEditorData(EditorData);
+						MLDeformerModel->CreateDataSetInterface();
+
+						// Train ML deformer model using user-defined parameters.
+						const double StartTime = FPlatformTime::Seconds();
+						const int ReturnCode = MLDeformerModel->Train();
+						const double TrainingDuration = FPlatformTime::Seconds() - StartTime;
+						const ETrainingResult TrainingResult = static_cast<ETrainingResult>(ReturnCode);
+						const bool bMarkDirty = HandleTrainingResult(TrainingResult, TrainingDuration);
+
+						UDebugSkelMeshComponent* SkelMeshComponent = EditorData->GetEditorActor(EMLDeformerEditorActorIndex::DeformedTest).SkelMeshComponent;
+						EditorData->InitAssets();
+						EditorData->GetEditorActor(EMLDeformerEditorActorIndex::DeformedTest).MLDeformerComponent->SetupComponent(DeformerAsset, SkelMeshComponent);
+						if (bMarkDirty)
+						{
+							EditorData->GetDeformerAsset()->Modify();
+						}
+						EditorData->GetDetailsView()->ForceRefresh();
+						EditorData->GetVizSettingsDetailsView()->ForceRefresh();
+					}
+					else
+					{
+						ShowNotification(LOCTEXT("ModelError", "Python Training module not defined by init_unreal.py"), SNotificationItem::ECompletionState::CS_Fail, true);
+						UE_LOG(LogMLDeformer, Error, TEXT("FMLDeformerEditorToolkit: Python Training module not defined by init_unreal.py"));
+					}
+				}),
+				FCanExecuteAction::CreateLambda(
+					[this]() 
+					{
+						return EditorData->IsReadyForTraining();
+					})
+			),
+			NAME_None,
+			LOCTEXT("TrainModel", "Train Model"),
+			LOCTEXT("TrainModelTooltip", "Train Model using Pytorch"),
+			FSlateIcon(),
+			EUserInterfaceActionType::ToggleButton
+		);
+	}
+	ToolbarBuilder.EndSection();
+}
+
+bool FMLDeformerEditorToolkit::HandleTrainingResult(ETrainingResult TrainingResult, double TrainingDuration)
+{
+	FText WindowTitle(LOCTEXT("TrainingResultsWindowTitle", "Training Results"));
+	FText WindowMessage;
+
+	// Calculate hours, minutes and seconds.
+    const int32 Hours = static_cast<int32>(TrainingDuration / 3600);
+    TrainingDuration -= Hours * 3600.0;
+    const int32 Minutes = static_cast<int32>(TrainingDuration / 60);
+    TrainingDuration -= Minutes * 60.0;
+    const int32 Seconds = static_cast<int32>(TrainingDuration);
+
+	// Format the results in some HH:MM::SS format.
+	FFormatNamedArguments Args;
+	FNumberFormattingOptions NumberOptions;
+	NumberOptions.SetMinimumIntegralDigits(2);
+	NumberOptions.SetUseGrouping(false);
+	Args.Add(TEXT("Hours"), FText::AsNumber(Hours, &NumberOptions));
+	Args.Add(TEXT("Minutes"), FText::AsNumber(Minutes, &NumberOptions));
+	Args.Add(TEXT("Seconds"), FText::AsNumber(Seconds, &NumberOptions));
+	const FText FormatString = LOCTEXT("TrainingDurationFormat", "{Hours}:{Minutes}:{Seconds} (HH:MM:SS)");
+	const FText TrainingDurationText = FText::Format(FormatString, Args);
+	UE_LOG(LogMLDeformer, Display, TEXT("Training duration: %s"), *TrainingDurationText.ToString());
+
+	bool bMarkDirty = false;
+	switch (TrainingResult)
+	{
+		// Training fully finished.
+		case ETrainingResult::Success:
+		{
+			if (!TryLoadOnnxFile())
+			{
+				GEditor->PlayEditorSound(TEXT("/Engine/EditorSounds/Notifications/CompileFailed_Cue.CompileFailed_Cue"));
+				WindowMessage = LOCTEXT("TrainingOnnxLoadFailed", "Training completed but resulting onnx file couldn't be loaded!");
+			}
+			else
+			{
+				GEditor->PlayEditorSound(TEXT("/Engine/EditorSounds/Notifications/CompileSuccess_Cue.CompileSuccess_Cue"));
+				FFormatNamedArguments SuccessArgs;
+				SuccessArgs.Add(TEXT("Duration"), TrainingDurationText);
+				WindowMessage = FText::Format(LOCTEXT("TrainingSuccess", "Training completed successfully!\n\nTraining time: {Duration}"), SuccessArgs);
+				bMarkDirty = true;
+			}
+		}
+		break;
+
+		// User aborted the training process. Ask whether they want to use the partially trained results or not.
+		case ETrainingResult::Aborted:
+		{
+			const FText Title(LOCTEXT("TrainingAbortedMessageTitle", "Use partially trained network?"));
+			const EAppReturnType::Type ReturnType = FMessageDialog::Open(
+				EAppMsgType::YesNo, 
+				FText(LOCTEXT("TrainingAbortedMessage", "Training has been aborted.\nThe neural network has only been partially trained.\nWould you like to use this partially trained network?")),
+				&Title);
+
+			if (ReturnType == EAppReturnType::Yes)
+			{
+				if (!TryLoadOnnxFile())
+				{
+					ShowNotification(LOCTEXT("TrainingOnnxLoadFailed", "Training partially completed, but resulting onnx file couldn't be loaded!"), SNotificationItem::ECompletionState::CS_Fail, true);
+				}
+				else
+				{
+					ShowNotification(LOCTEXT("PartialTrainingSuccess", "Training partially completed!"), SNotificationItem::ECompletionState::CS_Success, true);
+					bMarkDirty = true;
+				}
+			}
+			else
+			{
+				ShowNotification(LOCTEXT("TrainingAborted", "Training aborted!"), SNotificationItem::ECompletionState::CS_None, true);
+			}
+		}
+		break;
+
+		// Training data had issues.
+		case ETrainingResult::FailOnData:
+		{
+			GEditor->PlayEditorSound(TEXT("/Engine/EditorSounds/Notifications/CompileFailed_Cue.CompileFailed_Cue"));
+			WindowMessage = LOCTEXT("TrainingFailed", "Training failed!\nCheck input parameters or sequence length.");
+		}
+		break;
+
+		// Unknown failure, probably something in the python code.
+		case ETrainingResult::FailUnknown:
+		{
+			GEditor->PlayEditorSound(TEXT("/Engine/EditorSounds/Notifications/CompileFailed_Cue.CompileFailed_Cue"));
+			WindowMessage = LOCTEXT("TrainingFailed", "Training failed!\nUnknown error, please check the output log.");
+		}
+		break;
+
+		// Unhandled error codes.
+		default: check(false);
+	}
+
+	// Show a message window.
+	if (!WindowMessage.IsEmpty())
+	{
+		FMessageDialog::Open(EAppMsgType::Ok, WindowMessage, &WindowTitle);
+	}
+
+	return bMarkDirty;
+}
+
+bool FMLDeformerEditorToolkit::TryLoadOnnxFile() const
+{
+	FString OnnxFile = FPaths::ProjectIntermediateDir() + "MLDeformerModels/latest_net_G.onnx";
+	OnnxFile = FPaths::ConvertRelativePathToFull(OnnxFile);
+	if (FPaths::FileExists(OnnxFile))
+	{
+		UE_LOG(LogMLDeformer, Display, TEXT("Loading Onnx file '%s'..."), *OnnxFile);
+		UNeuralNetwork* Network = NewObject<UNeuralNetwork>(EditorData->GetDeformerAsset(), UNeuralNetwork::StaticClass());		
+		if (Network->Load(OnnxFile))
+		{
+			EditorData->GetDeformerAsset()->NeuralNetwork = Network;
+			EditorData->GetDeformerAsset()->NeuralNetwork->SetBackEnd(ENeuralBackEnd::UEOnly);
+			EditorData->GetDeformerAsset()->NeuralNetwork->SetInputDeviceType(ENeuralDeviceType::CPU);
+			EditorData->GetDeformerAsset()->NeuralNetwork->SetDeviceType(ENeuralDeviceType::GPU);
+			EditorData->GetDeformerAsset()->NeuralNetwork->SetOutputDeviceType(ENeuralDeviceType::GPU);
+
+			// Recreate the data providers.
+			// This is needed for the neural network GPU buffers to be valid.
+			SetComputeGraphDataProviders();
+
+			UE_LOG(LogMLDeformer, Display, TEXT("Successfully loaded Onnx file '%s'..."), *OnnxFile);
+			return true;
+		}
+		else
+		{
+			UE_LOG(LogMLDeformer, Error, TEXT("Failed to load Onnx file '%s'"), *OnnxFile);
+		}
+	}
+	else
+	{
+		UE_LOG(LogMLDeformer, Error, TEXT("Onnx file '%s' does not exist!"), *OnnxFile);
+	}
+
+	return false;
+}
+
+FName FMLDeformerEditorToolkit::GetToolkitFName() const
+{
+	return FName("MLDeformerEditor");
+}
+
+FText FMLDeformerEditorToolkit::GetBaseToolkitName() const
+{
+	return LOCTEXT("MLDeformerEditorAppLabel", "ML Deformer Editor");
+}
+
+void FMLDeformerEditorToolkit::ShowNotification(const FText& Message, SNotificationItem::ECompletionState State, bool PlaySound) const
+{
+	FNotificationInfo Info(Message);
+	Info.FadeInDuration = 0.1f;
+	Info.FadeOutDuration = 0.5f;
+	Info.ExpireDuration = 3.5f;
+	Info.bUseThrobber = false;
+	Info.bUseSuccessFailIcons = true;
+	Info.bUseLargeFont = true;
+	Info.bFireAndForget = true;
+	Info.bAllowThrottleWhenFrameRateIsLow = false;
+	auto NotificationItem = FSlateNotificationManager::Get().AddNotification(Info);
+	NotificationItem->SetCompletionState(State);
+	NotificationItem->ExpireAndFadeout();
+
+	if (PlaySound)
+	{
+		switch (State)
+		{
+			case SNotificationItem::ECompletionState::CS_Success: GEditor->PlayEditorSound(TEXT("/Engine/EditorSounds/Notifications/CompileSuccess_Cue.CompileSuccess_Cue")); break;
+			case SNotificationItem::ECompletionState::CS_Fail: GEditor->PlayEditorSound(TEXT("/Engine/EditorSounds/Notifications/CompileFailed_Cue.CompileFailed_Cue")); break;
+			case SNotificationItem::ECompletionState::CS_Pending: GEditor->PlayEditorSound(TEXT("/Engine/EditorSounds/Notifications/CompileStart_Cue.CompileStart_Cue")); break;
+			default:;
+		};
+	}
+}
+
+FText FMLDeformerEditorToolkit::GetToolkitName() const
+{
+	FFormatNamedArguments Args;
+	Args.Add(TEXT("AssetName"), FText::FromString(EditorData->GetDeformerAsset()->GetName()));
+	return FText::Format(LOCTEXT("DemoEditorToolkitName", "{AssetName}"), Args);
+}
+
+FLinearColor FMLDeformerEditorToolkit::GetWorldCentricTabColorScale() const
+{
+	return FLinearColor::White;
+}
+
+FString FMLDeformerEditorToolkit::GetWorldCentricTabPrefix() const
+{
+	return TEXT("MLDeformerEditor");
+}
+
+void FMLDeformerEditorToolkit::AddReferencedObjects(FReferenceCollector& Collector)
+{
+	UMLDeformerAsset* Asset = EditorData->GetDeformerAsset();
+	Collector.AddReferencedObject(Asset);
+}
+
+TStatId FMLDeformerEditorToolkit::GetStatId() const
+{
+	RETURN_QUICK_DECLARE_CYCLE_STAT(FMLDeformerEditorToolkit, STATGROUP_Tickables);
+}
+
+UTextRenderComponent* FMLDeformerEditorToolkit::CreateLabelForActor(AActor* Actor, UWorld* World, FLinearColor Color, const FText& Text) const
+{
+	const float DefaultLabelScale = FMLDeformerEditorStyle::Get().GetFloat("MLDeformer.DefaultLabelScale");
+	UTextRenderComponent* TargetLabelComponent = NewObject<UTextRenderComponent>(Actor);
+	TargetLabelComponent->SetMobility(EComponentMobility::Movable);
+	TargetLabelComponent->SetHorizontalAlignment(EHTA_Center);
+	TargetLabelComponent->SetVerticalAlignment(EVRTA_TextCenter);
+	TargetLabelComponent->SetText(Text);
+	TargetLabelComponent->SetRelativeScale3D(FVector(DefaultLabelScale));
+	TargetLabelComponent->SetGenerateOverlapEvents(false);
+	TargetLabelComponent->SetCanEverAffectNavigation(false);
+	TargetLabelComponent->SetTextRenderColor(Color.ToFColor(true));
+	TargetLabelComponent->RegisterComponent();
+	return TargetLabelComponent;
+}
+
+void FMLDeformerEditorToolkit::CreateSkinnedActor(EMLDeformerEditorActorIndex ActorIndex, const FName& Name, UWorld* World, USkeletalMesh* Mesh, FLinearColor LabelColor, FLinearColor WireframeColor) const
+{
+	// Spawn the actor.
+	FActorSpawnParameters SpawnParams;
+	SpawnParams.Name = Name;
+	AActor* Actor = World->SpawnActor<AActor>(SpawnParams);
+	Actor->SetFlags(RF_Transient);
+
+	// Create the skeletal mesh component.
+	UDebugSkelMeshComponent* SkelMeshComponent = NewObject<UDebugSkelMeshComponent>(Actor);
+	SkelMeshComponent->SetSkeletalMesh(Mesh);
+	Actor->SetRootComponent(SkelMeshComponent);
+	SkelMeshComponent->RegisterComponent();
+	SkelMeshComponent->SetWireframeMeshOverlayColor(WireframeColor);
+	SkelMeshComponent->MarkRenderStateDirty();
+
+	// Register the editor actor.
+	FMLDeformerEditorActor EditorActor;
+	EditorActor.Actor = Actor;
+	EditorActor.SkelMeshComponent = SkelMeshComponent;
+	EditorActor.LabelComponent = CreateLabelForActor(Actor, World, LabelColor, FText::FromString(Name.ToString()));
+	EditorData->SetEditorActor(ActorIndex, EditorActor);
+}
+
+UMLDeformerComponent* FMLDeformerEditorToolkit::AddMLDeformerComponentToActor(EMLDeformerEditorActorIndex ActorIndex)
+{
+	FMLDeformerEditorActor& EditorActor = EditorData->GetEditorActor(ActorIndex);
+	AActor* Actor = EditorActor.Actor;
+	UMLDeformerComponent* Component = NewObject<UMLDeformerComponent>(Actor);
+	Component->SetDeformerAsset(EditorData->GetDeformerAsset());
+	Component->RegisterComponent();
+	EditorActor.MLDeformerComponent = Component;
+	return Component;
+}
+
+void FMLDeformerEditorToolkit::AddDeformerGraphComponentToActor(EMLDeformerEditorActorIndex ActorIndex, UComputeGraph* ComputeGraph, UMLDeformerComponent* DeformerComponent)
+{
+	FMLDeformerEditorActor& EditorActor = EditorData->GetEditorActor(ActorIndex);
+	AActor* Actor = EditorActor.Actor;
+	UComputeGraphComponent* Component = NewObject<UComputeGraphComponent>(Actor);
+	Component->ComputeGraph = ComputeGraph;
+	Component->CreateDataProviders(true);
+	Component->PrimaryComponentTick.SetTickFunctionEnable(true);
+	if (DeformerComponent)
+	{
+		Component->AddTickPrerequisiteComponent(DeformerComponent);
+	}
+	Component->RegisterComponent();
+	EditorActor.ComputeGraphComponent = Component;
+}
+
+void FMLDeformerEditorToolkit::CreateBaseActor(const TSharedRef<IPersonaPreviewScene>& InPersonaPreviewScene, const FName& Name, FLinearColor LabelColor, FLinearColor WireframeColor)
+{
+	UWorld* World = InPersonaPreviewScene->GetWorld();
+
+	// Spawn the linear skinned actor.
+	FActorSpawnParameters BaseSpawnParams;
+	BaseSpawnParams.Name = Name;
+	AAnimationEditorPreviewActor* Actor = World->SpawnActor<AAnimationEditorPreviewActor>(AAnimationEditorPreviewActor::StaticClass(), FTransform::Identity, BaseSpawnParams);
+	Actor->SetFlags(RF_Transient);
+	InPersonaPreviewScene->SetActor(Actor);
+
+	// Create the preview skeletal mesh component.
+	UDebugSkelMeshComponent* SkelMeshComponent = NewObject<UDebugSkelMeshComponent>(Actor);
+	SkelMeshComponent->SetWireframeMeshOverlayColor(WireframeColor);
+	SkelMeshComponent->MarkRenderStateDirty();
+
+	// Setup an apply an anim instance to the skeletal mesh component.
+	UAnimPreviewInstance* AnimPreviewInstance = NewObject<UAnimPreviewInstance>(SkelMeshComponent, TEXT("MLDeformerAnimInstance"));
+	SkelMeshComponent->PreviewInstance = AnimPreviewInstance;
+	AnimPreviewInstance->InitializeAnimation();
+
+	// Set the skeletal mesh on the component.
+	// NOTE: This must be done AFTER setting the AnimInstance so that the correct root anim node is loaded.
+	USkeletalMesh* Mesh = EditorData->GetDeformerAsset()->GetSkeletalMesh();
+	SkelMeshComponent->SetSkeletalMesh(Mesh);
+
+	// Apply mesh to the preview scene.
+	InPersonaPreviewScene->SetPreviewMeshComponent(SkelMeshComponent);
+	InPersonaPreviewScene->AddComponent(SkelMeshComponent, FTransform::Identity);
+	InPersonaPreviewScene->SetAdditionalMeshesSelectable(false);
+	InPersonaPreviewScene->SetPreviewMesh(Mesh);
+
+	// Register the editor actor.
+	FMLDeformerEditorActor BaseEditorActor;
+	BaseEditorActor.Actor = Actor;
+	BaseEditorActor.SkelMeshComponent = SkelMeshComponent;
+	BaseEditorActor.LabelComponent = CreateLabelForActor(Actor, World, LabelColor, FText::FromString(Name.ToString()));
+	EditorData->SetEditorActor(EMLDeformerEditorActorIndex::Base, BaseEditorActor);
+}
+
+void FMLDeformerEditorToolkit::CreateGeomCacheActor(EMLDeformerEditorActorIndex ActorIndex, UWorld* World, const FName& Name, UGeometryCache* GeomCache, FLinearColor LabelColor, FLinearColor WireframeColor)
+{
+	FActorSpawnParameters SpawnParams;
+	SpawnParams.Name = Name;
+	AActor* Actor = World->SpawnActor<AActor>(SpawnParams);
+	Actor->SetFlags(RF_Transient);
+
+	// Create the Geometry Cache Component.
+	UGeometryCacheComponent* GeomCacheComponent = NewObject<UGeometryCacheComponent>(Actor);
+	GeomCacheComponent->SetGeometryCache(GeomCache);
+	GeomCacheComponent->RegisterComponent();
+	GeomCacheComponent->SetOverrideWireframeColor(true);
+	GeomCacheComponent->SetWireframeOverrideColor(WireframeColor);
+	GeomCacheComponent->MarkRenderStateDirty();
+	Actor->SetRootComponent(GeomCacheComponent);
+
+	UTextRenderComponent* TargetLabelComponent = CreateLabelForActor(Actor, World, LabelColor, FText::FromString(Name.ToString()));
+
+	// Register the editor actor.
+	FMLDeformerEditorActor EditorActor;
+	EditorActor.Actor = Actor;
+	EditorActor.GeomCacheComponent = GeomCacheComponent;
+	EditorActor.LabelComponent = TargetLabelComponent;
+	EditorData->SetEditorActor(ActorIndex, EditorActor);
+}
+
+void FMLDeformerEditorToolkit::HandlePreviewSceneCreated(const TSharedRef<IPersonaPreviewScene>& InPersonaPreviewScene)
+{
+	// Load the default ML deformer graph asset.
+	EditorData->SetDefaultDeformerGraphIfNeeded();
+
+	// Create the Linear Skinned (Base) actor.
+	const FLinearColor BaseLabelColor = FMLDeformerEditorStyle::Get().GetColor("MLDeformer.BaseMesh.LabelColor");
+	const FLinearColor BaseWireColor = FMLDeformerEditorStyle::Get().GetColor("MLDeformer.BaseMesh.WireframeColor");
+	CreateBaseActor(InPersonaPreviewScene, "Training Base", BaseLabelColor, BaseWireColor);
+
+	// Create the target actor (with Geometry Cache).
+	UWorld* World = InPersonaPreviewScene->GetWorld();
+	const FLinearColor TargetLabelColor = FMLDeformerEditorStyle::Get().GetColor("MLDeformer.TargetMesh.LabelColor");
+	const FLinearColor TargetWireColor = FMLDeformerEditorStyle::Get().GetColor("MLDeformer.TargetMesh.WireframeColor");
+	CreateGeomCacheActor(
+		EMLDeformerEditorActorIndex::Target, 
+		World, 
+		"Training Target", 
+		EditorData->GetDeformerAsset()->GetGeometryCache(), 
+		TargetLabelColor,
+		TargetWireColor);
+
+	// Create the linear skinned actor.
+	USkeletalMesh* Mesh = EditorData->GetDeformerAsset()->GetSkeletalMesh();
+	CreateSkinnedActor(EMLDeformerEditorActorIndex::Test, "Linear Skinned", World, Mesh, BaseLabelColor, BaseWireColor);
+
+	// Create the ML deformed actor.
+	const FLinearColor MLDeformedLabelColor = FMLDeformerEditorStyle::Get().GetColor("MLDeformer.MLDeformedMesh.LabelColor");
+	const FLinearColor MLDeformedWireColor = FMLDeformerEditorStyle::Get().GetColor("MLDeformer.MLDeformedMesh.WireframeColor");
+	CreateSkinnedActor(EMLDeformerEditorActorIndex::DeformedTest, "ML Deformed", World, Mesh, MLDeformedLabelColor, MLDeformedWireColor);
+	UDebugSkelMeshComponent* SkelMeshComponent = EditorData->GetEditorActor(EMLDeformerEditorActorIndex::DeformedTest).SkelMeshComponent;
+	UMLDeformerComponent* MLDeformerComponent = AddMLDeformerComponentToActor(EMLDeformerEditorActorIndex::DeformedTest);
+	MLDeformerComponent->SetupComponent(EditorData->GetDeformerAsset(), SkelMeshComponent);
+	SkelMeshComponent->SkinCacheUsage.SetNum(SkelMeshComponent->LODInfo.Num());
+	for (int32 Index = 0; Index < SkelMeshComponent->LODInfo.Num(); ++Index)
+	{
+		SkelMeshComponent->SkinCacheUsage[Index] = ESkinCacheUsage::Enabled;
+	}
+
+	// Create the component with the deformer graph on it.
+	UMLDeformerVizSettings* VizSettings = EditorData->GetDeformerAsset()->GetVizSettings();
+	UComputeGraph* ComputeGraph = VizSettings->GetDeformerGraph();
+	AddDeformerGraphComponentToActor(EMLDeformerEditorActorIndex::DeformedTest, ComputeGraph, MLDeformerComponent);
+
+	// Create the ground truth actor.
+	const FLinearColor GroundTruthLabelColor = FMLDeformerEditorStyle::Get().GetColor("MLDeformer.GroundTruth.LabelColor");
+	const FLinearColor GroundTruthWireColor = FMLDeformerEditorStyle::Get().GetColor("MLDeformer.GroundTruth.WireframeColor");
+	CreateGeomCacheActor(
+		EMLDeformerEditorActorIndex::GroundTruth, 
+		World, 
+		"Ground Truth", 
+		EditorData->GetDeformerAsset()->GetGeometryCache(), 
+		GroundTruthLabelColor,
+		GroundTruthWireColor);
+
+	// Start playing the animations.
+	EditorData->InitAssets();
+
+	// Create visualization materials.
+	EditorData->CreateMaterials();
+	EditorData->SetHeatMapMaterialEnabled(EditorData->GetDeformerAsset()->GetVizSettings()->GetShowHeatMap());
+
+	VizSettings->SetTempVisualizationMode(VizSettings->GetVisualizationMode());
+	EditorData->GetDeformerAsset()->SetTempTrainingInputs(EditorData->GetDeformerAsset()->GetTrainingInputs());
+	EditorData->UpdateDebugPointData();
+	OnSwitchedVisualizationMode();
+}
+
+void FMLDeformerEditorToolkit::HandleDetailsCreated(const TSharedRef<class IDetailsView>& InDetailsView)
+{
+	EditorData->SetDetailsView(InDetailsView);
+	InDetailsView->OnFinishedChangingProperties().AddSP(this, &FMLDeformerEditorToolkit::OnFinishedChangingDetails);
+	InDetailsView->SetObject(EditorData->GetDeformerAsset());
+}
+
+void FMLDeformerEditorToolkit::OnFinishedChangingDetails(const FPropertyChangedEvent& PropertyChangedEvent)
+{
+	const FProperty* Property = PropertyChangedEvent.Property;
+	if (Property == nullptr)
+	{
+		return;
+	}
+
+	// When we change one of these properties below, restart animations etc.
+	if (Property->GetFName() == GET_MEMBER_NAME_CHECKED(UMLDeformerAsset, SkeletalMesh))
+	{
+		EditorData->InitAssets();
+
+		SetComputeGraphDataProviders();
+
+		FMLDeformerEditorActor& DeformedTestActor = EditorData->GetEditorActor(EMLDeformerEditorActorIndex::DeformedTest);
+		USkeletalMeshComponent* SkelMeshComponent = DeformedTestActor.SkelMeshComponent;
+		UMLDeformerComponent* MLDeformerComponent = DeformedTestActor.MLDeformerComponent;
+		UComputeGraphComponent* ComputeComponent = DeformedTestActor.ComputeGraphComponent;
+
+		MLDeformerComponent->SetupComponent(EditorData->GetDeformerAsset(), SkelMeshComponent);
+
+		const bool IsCompatible = MLDeformerComponent->GetDeformerInstance().IsCompatible();	
+		if (ComputeComponent->IsActive() != IsCompatible)
+		{
+			ComputeComponent->SetActive(IsCompatible);
+		}
+
+		EditorData->GetDetailsView()->ForceRefresh();
+	}
+	else
+	if (Property->GetFName() == GET_MEMBER_NAME_CHECKED(UMLDeformerAsset, GeometryCache) ||
+		Property->GetFName() == GET_MEMBER_NAME_CHECKED(UMLDeformerAsset, AnimSequence) ||
+		Property->GetFName() == GET_MEMBER_NAME_CHECKED(UMLDeformerVizSettings, TestAnimSequence) ||
+		Property->GetFName() == GET_MEMBER_NAME_CHECKED(UMLDeformerVizSettings, GroundTruth))
+	{
+		EditorData->InitAssets();
+		EditorData->GetDetailsView()->ForceRefresh();
+		EditorData->GetVizSettingsDetailsView()->ForceRefresh();
+	}
+	else
+	if (Property->GetFName() == GET_MEMBER_NAME_CHECKED(UMLDeformerAsset, TrainingInputs))
+	{
+		UMLDeformerAsset* DeformerAsset = EditorData->GetDeformerAsset();
+		if (DeformerAsset->GetTempTrainingInputs() != DeformerAsset->GetTrainingInputs())
+		{
+			DeformerAsset->SetTempTrainingInputs(DeformerAsset->GetTrainingInputs());
+			EditorData->GetDetailsView()->ForceRefresh();
+		}
+	}
+	else
+	if (Property->GetFName() == GET_MEMBER_NAME_CHECKED(UMLDeformerVizSettings, AnimPlaySpeed))
+	{
+		EditorData->UpdateTestAnimPlaySpeed();
+	}
+	else
+	if (Property->GetFName() == GET_MEMBER_NAME_CHECKED(UMLDeformerVizSettings, DeformerGraph))
+	{
+		EditorData->UpdateDeformerGraph();
+		SetComputeGraphDataProviders();
+		EditorData->GetVizSettingsDetailsView()->ForceRefresh();
+	}
+	else
+	if (Property->GetFName() == GET_MEMBER_NAME_CHECKED(UMLDeformerVizSettings, bShowHeatMap))
+	{
+		EditorData->SetHeatMapMaterialEnabled(EditorData->GetDeformerAsset()->GetVizSettings()->GetShowHeatMap());
+	}
+	else
+	if (Property->GetFName() == GET_MEMBER_NAME_CHECKED(UMLDeformerVizSettings, VisualizationMode))
+	{
+		OnSwitchedVisualizationMode();
+	}
+	else
+	if (Property->GetFName() == GET_MEMBER_NAME_CHECKED(UMLDeformerVizSettings, bDrawLinearSkinnedActor) ||
+		Property->GetFName() == GET_MEMBER_NAME_CHECKED(UMLDeformerVizSettings, bDrawMLDeformedActor) ||
+		Property->GetFName() == GET_MEMBER_NAME_CHECKED(UMLDeformerVizSettings, bDrawGroundTruthActor))
+	{
+		UpdateActorVisibility();
+	}
+}
+
+
+void FMLDeformerEditorToolkit::UpdateActorVisibility()
+{
+	UMLDeformerAsset* DeformerAsset = EditorData->GetDeformerAsset();
+	UMLDeformerVizSettings* VizSettings = DeformerAsset->GetVizSettings();
+
+	// Change visibility.
+	const bool bShowTrainingData = (VizSettings->GetVisualizationMode() == EMLDeformerVizMode::TrainingData);
+	const bool bShowTestData = (VizSettings->GetVisualizationMode() == EMLDeformerVizMode::TestData);
+	EditorData->SetActorVisibility(EMLDeformerEditorActorIndex::Base, bShowTrainingData);
+	EditorData->SetActorVisibility(EMLDeformerEditorActorIndex::Target, bShowTrainingData);
+	EditorData->SetActorVisibility(EMLDeformerEditorActorIndex::Test, bShowTestData && VizSettings->GetDrawLinearSkinnedActor());
+	EditorData->SetActorVisibility(EMLDeformerEditorActorIndex::DeformedTest, bShowTestData && VizSettings->GetDrawMLDeformedActor());
+	EditorData->SetActorVisibility(EMLDeformerEditorActorIndex::GroundTruth, bShowTestData && VizSettings->GetDrawGroundTruthActor());
+}
+
+void FMLDeformerEditorToolkit::OnSwitchedVisualizationMode()
+{
+	UMLDeformerAsset* DeformerAsset = EditorData->GetDeformerAsset();
+	UMLDeformerVizSettings* VizSettings = DeformerAsset->GetVizSettings();
+
+	UpdateActorVisibility();
+
+	// Only trigger a force refresh when the value really changed.
+	// This is required due to a bug in the UI system that triggers a changed event already when clicking the combo box.
+	if (VizSettings->GetTempVisualizationMode() != VizSettings->GetVisualizationMode())
+	{
+		EditorData->GetVizSettingsDetailsView()->ForceRefresh();
+		VizSettings->SetTempVisualizationMode(VizSettings->GetVisualizationMode());
+	}
+
+	// Make sure the time slider is updated to reflect the right animation range.
+	EditorData->UpdateTimeSlider();
+}
+
+void FMLDeformerEditorToolkit::HandleViewportCreated(const TSharedRef<IPersonaViewport>& InPersonaViewport)
+{
+	InPersonaViewport->AddOverlayWidget(
+		SNew(SVerticalBox)
+		+SVerticalBox::Slot()
+		.Padding(FMargin(5.0f, 40.0f))
+		[
+			SNew(STextBlock)
+			.Text(this, &FMLDeformerEditorToolkit::GetOverlayText)
+			.Visibility(EVisibility::HitTestInvisible)
+			.ColorAndOpacity(FLinearColor(1.0f, 0.0f, 0.0f, 1.0f))
+		]
+	);
+}
+
+FText FMLDeformerEditorToolkit::GetOverlayText() const
+{
+	if (EditorData == nullptr)
+	{
+		return FText::GetEmpty();
+	}
+
+	return EditorData->GetOverlayText();
+}
+
+void FMLDeformerEditorToolkit::SetComputeGraphDataProviders() const
+{
+	const FMLDeformerEditorActor& EditorActor = EditorData->GetEditorActor(EMLDeformerEditorActorIndex::DeformedTest);
+	UComputeGraphComponent* Component = EditorActor.ComputeGraphComponent;
+	if (Component != nullptr)
+	{
+		Component->CreateDataProviders(true);
+	}
+}
+
+#undef LOCTEXT_NAMESPACE
diff --git a/Engine/Plugins/Experimental/MLDeformer/Source/MLDeformerEditor/Private/MLDeformerPythonTrainingModel.cpp b/Engine/Plugins/Experimental/MLDeformer/Source/MLDeformerEditor/Private/MLDeformerPythonTrainingModel.cpp
new file mode 100644
--- /dev/null
+++ b/Engine/Plugins/Experimental/MLDeformer/Source/MLDeformerEditor/Private/MLDeformerPythonTrainingModel.cpp
@@ -0,0 +1,42 @@
+﻿// Copyright Epic Games, Inc. All Rights Reserved.
+
+#include "MLDeformerPythonTrainingModel.h"
+#include "UObject/UObjectHash.h"
+#include "MLDeformerAsset.h"
+#include "MLDeformerEditorData.h"
+#include "MLPytorchDataSetInterface.h"
+
+UMLDeformerPythonTrainingModel::UMLDeformerPythonTrainingModel()
+{
+}
+
+UMLDeformerPythonTrainingModel* UMLDeformerPythonTrainingModel::Get()
+{
+	TArray<UClass*> PythonTrainingModels;
+	GetDerivedClasses(UMLDeformerPythonTrainingModel::StaticClass(), PythonTrainingModels);
+	const int32 NumClasses = PythonTrainingModels.Num();
+	if (NumClasses > 0)
+	{
+		return Cast<UMLDeformerPythonTrainingModel>(PythonTrainingModels[NumClasses - 1]->GetDefaultObject());
+	}
+	return nullptr;
+}
+
+UMLDeformerAsset* UMLDeformerPythonTrainingModel::GetDeformerAsset()
+{
+	return EditorData->GetDeformerAsset();
+}
+
+void UMLDeformerPythonTrainingModel::CreateDataSetInterface()
+{
+	if (DataSetInterface == nullptr)
+	{
+		DataSetInterface = NewObject<UMLPytorchDataSetInterface>();
+	}
+	DataSetInterface->SetEditorData(EditorData);
+}
+
+void UMLDeformerPythonTrainingModel::SetEditorData(TSharedPtr<FMLDeformerEditorData> InEditorData)
+{
+	EditorData = InEditorData;
+}
diff --git a/Engine/Plugins/Experimental/MLDeformer/Source/MLDeformerEditor/Private/MLDeformerPythonTrainingModel.h b/Engine/Plugins/Experimental/MLDeformer/Source/MLDeformerEditor/Private/MLDeformerPythonTrainingModel.h
new file mode 100644
--- /dev/null
+++ b/Engine/Plugins/Experimental/MLDeformer/Source/MLDeformerEditor/Private/MLDeformerPythonTrainingModel.h
@@ -0,0 +1,61 @@
+﻿// Copyright Epic Games, Inc. All Rights Reserved.
+#pragma once
+
+#include "CoreMinimal.h"
+#include "UObject/Object.h"
+#include "MLDeformer/Public/MLDeformerAsset.h"
+#include "MLDeformerPythonTrainingModel.generated.h"
+
+class UMLPytorchDataSetInterface;
+class FMLDeformerEditorData;
+class UMLDeformerAsset;
+
+/** Training process return codes. */
+UENUM()
+enum class ETrainingResult : uint8
+{
+	Success = 0,	/** The training successfully finished. */
+	Aborted,		/** The user has aborted the training process. */
+	FailOnData,		/** The input or output data to the network has issues, which means we cannot train. */
+	FailUnknown		/** There is an unknown error (see output log). */
+};
+
+/**
+ * 
+ */
+UCLASS(Blueprintable)
+class MLDEFORMEREDITOR_API UMLDeformerPythonTrainingModel : public UObject
+{
+	GENERATED_BODY()
+public:
+	UMLDeformerPythonTrainingModel();
+
+	UFUNCTION(BlueprintCallable, Category = Python)
+	static UMLDeformerPythonTrainingModel* Get();
+
+	/** 
+	 * Train model using training settings and metadata from the interface.
+	 * This function is implemented in the python scripts.
+	 */
+	UFUNCTION(BlueprintImplementableEvent, Category = Python)
+	int32 Train() const;
+
+	/** 
+	 * Create the dataset interface from the editor data.
+	 * Such an interface provides the samples for training the model.
+	 */
+	UFUNCTION(BlueprintCallable, Category = DataSet)
+	void CreateDataSetInterface();
+
+	/** Get the current ML deformer asset. **/
+	UFUNCTION(BlueprintCallable, Category = Python)
+	UMLDeformerAsset* GetDeformerAsset();
+
+	UPROPERTY(BlueprintReadOnly, Category = Python)
+	UMLPytorchDataSetInterface* DataSetInterface = nullptr;
+
+	void SetEditorData(TSharedPtr<FMLDeformerEditorData> InEditorData);
+
+private:
+	TSharedPtr<FMLDeformerEditorData> EditorData = nullptr;
+};
diff --git a/Engine/Plugins/Experimental/MLDeformer/Source/MLDeformerEditor/Private/MLDeformerTimelineTabSummoner.cpp b/Engine/Plugins/Experimental/MLDeformer/Source/MLDeformerEditor/Private/MLDeformerTimelineTabSummoner.cpp
new file mode 100644
--- /dev/null
+++ b/Engine/Plugins/Experimental/MLDeformer/Source/MLDeformerEditor/Private/MLDeformerTimelineTabSummoner.cpp
@@ -0,0 +1,91 @@
+// Copyright Epic Games, Inc. All Rights Reserved.
+
+#include "MLDeformerTimelineTabSummoner.h"
+#include "MLDeformerEditorToolkit.h"
+#include "MLDeformerEditorData.h"
+#include "MLDeformerEditorStyle.h"
+
+#include "IDocumentation.h"
+#include "SSimpleTimeSlider.h"
+#include "Widgets/SBoxPanel.h"
+#include "Widgets/Input/SButton.h"
+#include "Widgets/Images/SImage.h"
+
+#define LOCTEXT_NAMESPACE "MLDeformerTimelineTabSummoner"
+
+const FName FMLDeformerTimelineTabSummoner::TabID(TEXT("MLDeformerTimeline"));
+
+FMLDeformerTimelineTabSummoner::FMLDeformerTimelineTabSummoner(const TSharedRef<FMLDeformerEditorToolkit>& InEditor)
+	: FWorkflowTabFactory(TabID, InEditor)
+	, Editor(InEditor)
+{
+	bIsSingleton = true; // only allow a single instance of this tab
+
+	TabLabel = LOCTEXT("TimelineTabLabel", "Timeline");
+	TabIcon = FSlateIcon(FMLDeformerEditorStyle::Get().GetStyleSetName(), "MLDeformer.Timeline.TabIcon");
+	ViewMenuDescription = LOCTEXT("ViewMenu_Desc", "Timeline");
+	ViewMenuTooltip = LOCTEXT("ViewMenu_ToolTip", "Show the ML Deformer timeline.");
+}
+
+TSharedPtr<SToolTip> FMLDeformerTimelineTabSummoner::CreateTabToolTipWidget(const FWorkflowTabSpawnInfo& Info) const
+{
+	return IDocumentation::Get()->CreateToolTip(
+		LOCTEXT("TimelineTooltip", "The timeline widget that controls the offset in the training or test anim sequence."), 
+		NULL, 
+		TEXT("Shared/Editors/Persona"), 
+		TEXT("MLDeformerTimeline_Window"));
+}
+
+TSharedRef<SWidget> FMLDeformerTimelineTabSummoner::CreateTabBody(const FWorkflowTabSpawnInfo& Info) const
+{
+	// Create and setup the time slider widget.
+	FMLDeformerEditorData* EditorData = Editor.Pin()->GetEditorData();
+	TSharedRef<SSimpleTimeSlider> TimeSlider = SNew(SSimpleTimeSlider)
+		.ScrubPosition_Raw(EditorData, &FMLDeformerEditorData::CalcTimelinePosition)
+		.OnScrubPositionChanged_Raw(EditorData, &FMLDeformerEditorData::OnTimeSliderScrubPositionChanged);
+	EditorData->SetTimeSlider(TimeSlider);
+
+	// Add the time slider.
+	TSharedRef<SHorizontalBox> Content = SNew(SHorizontalBox)
+		+SHorizontalBox::Slot()
+		.VAlign(EVerticalAlignment::VAlign_Top)
+		[
+			TimeSlider
+		]
+		+SHorizontalBox::Slot()
+		.AutoWidth()
+		.HAlign(EHorizontalAlignment::HAlign_Right)
+		.VAlign(EVerticalAlignment::VAlign_Top)
+		[
+			SNew(SBox)
+			.WidthOverride(25)
+			.HeightOverride(25)
+			.Padding(FMargin(0.0f))
+			.Visibility_Lambda([EditorData](){ return (EditorData->GetDeformerAsset()->GetVizSettings()->GetVisualizationMode() == EMLDeformerVizMode::TestData) ? EVisibility::Visible : EVisibility::Collapsed; } )
+			[
+				SNew(SButton)
+				.ButtonStyle(FAppStyle::Get(), "SimpleButton")
+				.HAlign(EHorizontalAlignment::HAlign_Center)
+				.VAlign(EVerticalAlignment::VAlign_Center)
+				.ToolTipText(LOCTEXT("PlayButtonToolTip", "Play or pause the test animation sequence"))
+				.ContentPadding(FMargin(0.0f))
+				.OnClicked_Lambda([EditorData]() { EditorData->OnPlayButtonPressed(); return FReply::Handled(); })
+				[
+					SNew(SImage)
+					.Image_Lambda
+					(
+						[EditorData]()
+						{
+							return EditorData->IsPlayingAnim() 
+								? FMLDeformerEditorStyle::Get().GetBrush("MLDeformer.Timeline.PauseIcon")
+								: FMLDeformerEditorStyle::Get().GetBrush("MLDeformer.Timeline.PlayIcon");
+						}
+					)
+				]
+			]
+		];
+
+	return Content;
+}
+
+#undef LOCTEXT_NAMESPACE 
diff --git a/Engine/Plugins/Experimental/MLDeformer/Source/MLDeformerEditor/Private/MLDeformerTimelineTabSummoner.h b/Engine/Plugins/Experimental/MLDeformer/Source/MLDeformerEditor/Private/MLDeformerTimelineTabSummoner.h
new file mode 100644
--- /dev/null
+++ b/Engine/Plugins/Experimental/MLDeformer/Source/MLDeformerEditor/Private/MLDeformerTimelineTabSummoner.h
@@ -0,0 +1,23 @@
+// Copyright Epic Games, Inc. All Rights Reserved.
+
+#pragma once
+
+#include "CoreMinimal.h"
+#include "WorkflowOrientedApp/WorkflowTabFactory.h"
+
+class FMLDeformerEditorToolkit;
+class IDetailsView;
+
+struct FMLDeformerTimelineTabSummoner : public FWorkflowTabFactory
+{
+public:
+	static const FName TabID;
+
+	FMLDeformerTimelineTabSummoner(const TSharedRef<FMLDeformerEditorToolkit>& InEditor);
+	virtual TSharedRef<SWidget> CreateTabBody(const FWorkflowTabSpawnInfo& Info) const override;
+	virtual TSharedPtr<SToolTip> CreateTabToolTipWidget(const FWorkflowTabSpawnInfo& Info) const override;
+
+protected:
+	TWeakPtr<FMLDeformerEditorToolkit> Editor;
+	TSharedPtr<IDetailsView> DetailsView;
+};
diff --git a/Engine/Plugins/Experimental/MLDeformer/Source/MLDeformerEditor/Private/MLDeformerVizSettingsDetails.cpp b/Engine/Plugins/Experimental/MLDeformer/Source/MLDeformerEditor/Private/MLDeformerVizSettingsDetails.cpp
new file mode 100644
--- /dev/null
+++ b/Engine/Plugins/Experimental/MLDeformer/Source/MLDeformerEditor/Private/MLDeformerVizSettingsDetails.cpp
@@ -0,0 +1,266 @@
+// Copyright Epic Games, Inc. All Rights Reserved.
+
+#include "MLDeformerVizSettingsDetails.h"
+#include "MLDeformerVizSettings.h"
+#include "MLDeformer.h"
+#include "MLDeformerAsset.h"
+
+#include "MLDeformerEditorData.h"
+#include "ComputeFramework/ComputeGraph.h"
+
+#include "DetailLayoutBuilder.h"
+#include "DetailCategoryBuilder.h"
+#include "DetailWidgetRow.h"
+#include "IDetailsView.h"
+
+#include "Widgets/SWidget.h"
+#include "Layout/Margin.h"
+#include "Widgets/DeclarativeSyntaxSupport.h"
+#include "Widgets/SBoxPanel.h"
+#include "EditorStyleSet.h"
+#include "IDetailChildrenBuilder.h"
+#include "IDetailGroup.h"
+#include "Framework/Application/SlateApplication.h"
+#include "Widgets/Layout/SBorder.h"
+#include "Widgets/Text/STextBlock.h"
+#include "Widgets/Layout/SBox.h"
+#include "PropertyCustomizationHelpers.h"
+#include "AssetRegistry/AssetData.h"
+
+#include "PropertyHandle.h"
+#include "DetailLayoutBuilder.h"
+
+#include "SWarningOrErrorBox.h"
+
+#include "GeometryCache.h"
+#include "Animation/AnimSequence.h"
+#include "Engine/SkeletalMesh.h"
+
+#define LOCTEXT_NAMESPACE "MLDeformerVizSettingsDetails"
+
+TSharedRef<IDetailCustomization> FMLDeformerVizSettingsDetails::MakeInstance()
+{
+	return MakeShareable(new FMLDeformerVizSettingsDetails());
+}
+
+UMLDeformerAsset* FMLDeformerVizSettingsDetails::GetMLDeformerAsset() const
+{
+	TArray<TWeakObjectPtr<UObject>> Objects;
+	DetailLayoutBuilder->GetObjectsBeingCustomized(Objects);
+	UMLDeformerVizSettings* VizSettings = nullptr;
+	UMLDeformerAsset* DeformerAsset = nullptr;
+	if (Objects.Num() == 1)
+	{
+		VizSettings = Cast<UMLDeformerVizSettings>(Objects[0]);
+		check(VizSettings);
+		DeformerAsset = Cast<UMLDeformerAsset>(VizSettings->GetOuter());
+		check(DeformerAsset);
+		return DeformerAsset;
+	}
+
+	return nullptr;
+}
+
+void FMLDeformerVizSettingsDetails::CustomizeDetails(class IDetailLayoutBuilder& DetailBuilder)
+{
+	DetailLayoutBuilder = &DetailBuilder;
+
+	UMLDeformerAsset* DeformerAsset = GetMLDeformerAsset();
+	UMLDeformerVizSettings* VizSettings = DeformerAsset->GetVizSettings();
+
+	const bool bShowTrainingData = VizSettings ? (VizSettings->GetVisualizationMode() == EMLDeformerVizMode::TrainingData) : true;
+	const bool bShowTestData = VizSettings ? (VizSettings->GetVisualizationMode() == EMLDeformerVizMode::TestData) : true;
+
+	IDetailCategoryBuilder& DataCategoryBuilder = DetailBuilder.EditCategory("Data Selection", FText::GetEmpty(), ECategoryPriority::Important);
+	DataCategoryBuilder.AddProperty(GET_MEMBER_NAME_CHECKED(UMLDeformerVizSettings, VisualizationMode));
+
+	// Shared settings.
+	IDetailCategoryBuilder& SharedCategoryBuilder = DetailBuilder.EditCategory("Shared Settings", FText::GetEmpty(), ECategoryPriority::Important);
+	SharedCategoryBuilder.AddProperty(GET_MEMBER_NAME_CHECKED(UMLDeformerVizSettings, bDrawLabels));
+	SharedCategoryBuilder.AddProperty(GET_MEMBER_NAME_CHECKED(UMLDeformerVizSettings, LabelHeight));
+	SharedCategoryBuilder.AddProperty(GET_MEMBER_NAME_CHECKED(UMLDeformerVizSettings, LabelScale));
+	SharedCategoryBuilder.AddProperty(GET_MEMBER_NAME_CHECKED(UMLDeformerVizSettings, MeshSpacing));
+
+	// Testing.
+	IDetailCategoryBuilder& TestingCategoryBuilder = DetailBuilder.EditCategory("Testing", FText::GetEmpty(), ECategoryPriority::Important);
+	TestingCategoryBuilder.SetCategoryVisibility(bShowTestData);
+	
+	IDetailPropertyRow& TestAnimRow = TestingCategoryBuilder.AddProperty(GET_MEMBER_NAME_CHECKED(UMLDeformerVizSettings, TestAnimSequence));
+	TestAnimRow.CustomWidget()
+	.NameContent()
+	[
+		TestAnimRow.GetPropertyHandle()->CreatePropertyNameWidget()
+	]
+	.ValueContent()
+	[
+		SNew(SObjectPropertyEntryBox)
+		.PropertyHandle(TestAnimRow.GetPropertyHandle())
+		.AllowedClass(UAnimSequence::StaticClass())
+		.ObjectPath(VizSettings ? VizSettings->GetTestAnimSequence()->GetPathName() : FString())
+		.ThumbnailPool(DetailBuilder.GetThumbnailPool())
+		.OnShouldFilterAsset(
+			this, 
+			&FMLDeformerVizSettingsDetails::FilterAnimSequences, 
+			DeformerAsset->GetSkeletalMesh() ? DeformerAsset->GetSkeletalMesh()->GetSkeleton() : nullptr
+		)
+	];
+
+	if (VizSettings)
+	{
+		const FText AnimErrorText = DeformerAsset->GetIncompatibleSkeletonErrorText(DeformerAsset->GetSkeletalMesh(), VizSettings->GetTestAnimSequence());
+		FDetailWidgetRow& AnimErrorRow = TestingCategoryBuilder.AddCustomRow(FText::FromString("AnimSkeletonMisMatchError"))
+			.Visibility(!AnimErrorText.IsEmpty() ? EVisibility::Visible : EVisibility::Collapsed)
+			.WholeRowContent()
+			[
+				SNew(SBox)
+				.Padding(FMargin(0.0f, 4.0f))
+				[
+					SNew(SWarningOrErrorBox)
+					.MessageStyle(EMessageStyle::Warning)
+					.Message(AnimErrorText)
+				]
+			];
+	}
+
+	TestingCategoryBuilder.AddProperty(GET_MEMBER_NAME_CHECKED(UMLDeformerVizSettings, AnimPlaySpeed));
+
+	FIsResetToDefaultVisible IsResetVisible = FIsResetToDefaultVisible::CreateSP(this, &FMLDeformerVizSettingsDetails::IsResetToDefaultDeformerGraphVisible);
+	FResetToDefaultHandler ResetHandler = FResetToDefaultHandler::CreateSP(this, &FMLDeformerVizSettingsDetails::OnResetToDefaultDeformerGraph);
+	FResetToDefaultOverride ResetOverride = FResetToDefaultOverride::Create(IsResetVisible, ResetHandler);
+	TestingCategoryBuilder.AddProperty(GET_MEMBER_NAME_CHECKED(UMLDeformerVizSettings, DeformerGraph)).OverrideResetToDefault(ResetOverride);
+
+	// Show a warning when no deformer graph has been selected.
+	UObject* Graph = nullptr;
+	TSharedRef<IPropertyHandle> DeformerGraphProperty = DetailBuilder.GetProperty(GET_MEMBER_NAME_CHECKED(UMLDeformerVizSettings, DeformerGraph));
+	if (DeformerGraphProperty->GetValue(Graph) == FPropertyAccess::Result::Success)
+	{
+		FDetailWidgetRow& GraphErrorRow = TestingCategoryBuilder.AddCustomRow(FText::FromString("GraphError"))
+			.Visibility((Graph == nullptr) ? EVisibility::Visible : EVisibility::Collapsed)
+			.WholeRowContent()
+			[
+				SNew(SBox)
+				.Padding(FMargin(0.0f, 4.0f))
+				[
+					SNew(SWarningOrErrorBox)
+					.MessageStyle(EMessageStyle::Warning)
+					.Message(FText::FromString("Please select a deformer graph.\nOtherwise only linear skinning is used."))
+				]
+			];
+	}
+
+	if (DeformerAsset)
+	{
+		FDetailWidgetRow& ErrorRow = TestingCategoryBuilder.AddCustomRow(FText::FromString("NoNeuralNetError"))
+			.Visibility((DeformerAsset->GetInferenceNeuralNetwork() == nullptr && Graph != nullptr) ? EVisibility::Visible : EVisibility::Collapsed)
+			.WholeRowContent()
+			[
+				SNew(SBox)
+				.Padding(FMargin(0.0f, 4.0f))
+				[
+					SNew(SWarningOrErrorBox)
+					.MessageStyle(EMessageStyle::Warning)
+					.Message(FText::FromString("The selected deformer graph isn't used, because you didn't train the neural network yet.\n\nLinear skinning is used until then."))
+				]
+			];
+	}
+
+	TestingCategoryBuilder.AddProperty(GET_MEMBER_NAME_CHECKED(UMLDeformerVizSettings, GroundTruth));
+
+	// Show an error when the test anim sequence duration doesn't match the one of the ground truth.
+	if (VizSettings)
+	{
+		const FText AnimErrorText = DeformerAsset->GetAnimSequenceErrorText(VizSettings->GroundTruth, VizSettings->GetTestAnimSequence());
+		FDetailWidgetRow& GroundTruthAnimErrorRow = TestingCategoryBuilder.AddCustomRow(FText::FromString("GroundTruthAnimMismatchError"))
+			.Visibility(!AnimErrorText.IsEmpty() ? EVisibility::Visible : EVisibility::Collapsed)
+			.WholeRowContent()
+			[
+				SNew(SBox)
+				.Padding(FMargin(0.0f, 4.0f))
+				[
+					SNew(SWarningOrErrorBox)
+					.MessageStyle(EMessageStyle::Warning)
+					.Message(AnimErrorText)
+				]
+			];
+
+		const FText GeomErrorText = DeformerAsset->GetGeomCacheErrorText(VizSettings->GetGroundTruth());
+		FDetailWidgetRow& GroundTruthGeomErrorRow = TestingCategoryBuilder.AddCustomRow(FText::FromString("GroundTruthGeomMismatchError"))
+			.Visibility(!GeomErrorText.IsEmpty() ? EVisibility::Visible : EVisibility::Collapsed)
+			.WholeRowContent()
+			[
+				SNew(SBox)
+				.Padding(FMargin(0.0f, 4.0f))
+				[
+					SNew(SWarningOrErrorBox)
+					.MessageStyle(EMessageStyle::Warning)
+					.Message(GeomErrorText)
+				]
+			];
+
+		const FText VertexErrorText = DeformerAsset->GetVertexErrorText(DeformerAsset->SkeletalMesh, VizSettings->GetGroundTruth(), FText::FromString("Base Mesh"), FText::FromString("Ground Truth Mesh"));
+		FDetailWidgetRow& GroundTruthVertexErrorRow = TestingCategoryBuilder.AddCustomRow(FText::FromString("GroundTruthVertexMismatchError"))
+			.Visibility(!VertexErrorText.IsEmpty() ? EVisibility::Visible : EVisibility::Collapsed)
+			.WholeRowContent()
+			[
+				SNew(SBox)
+				.Padding(FMargin(0.0f, 4.0f))
+				[
+					SNew(SWarningOrErrorBox)
+					.MessageStyle(EMessageStyle::Warning)
+					.Message(VertexErrorText)
+				]
+			];
+	}
+
+	TestingCategoryBuilder.AddProperty(GET_MEMBER_NAME_CHECKED(UMLDeformerVizSettings, VertexDeltaMultiplier));
+	TestingCategoryBuilder.AddProperty(GET_MEMBER_NAME_CHECKED(UMLDeformerVizSettings, bShowHeatMap));
+	TestingCategoryBuilder.AddProperty(GET_MEMBER_NAME_CHECKED(UMLDeformerVizSettings, bDrawLinearSkinnedActor));
+	TestingCategoryBuilder.AddProperty(GET_MEMBER_NAME_CHECKED(UMLDeformerVizSettings, bDrawMLDeformedActor));
+	TestingCategoryBuilder.AddProperty(GET_MEMBER_NAME_CHECKED(UMLDeformerVizSettings, bDrawGroundTruthActor));
+
+	// Training data.
+	IDetailCategoryBuilder& TrainingMeshesCategoryBuilder = DetailBuilder.EditCategory("Training Meshes", FText::GetEmpty(), ECategoryPriority::Important);
+	TrainingMeshesCategoryBuilder.SetCategoryVisibility(bShowTrainingData);
+	TrainingMeshesCategoryBuilder.AddProperty(GET_MEMBER_NAME_CHECKED(UMLDeformerVizSettings, FrameNumber));
+	TrainingMeshesCategoryBuilder.AddProperty(GET_MEMBER_NAME_CHECKED(UMLDeformerVizSettings, VertexNumber));
+	TrainingMeshesCategoryBuilder.AddProperty(GET_MEMBER_NAME_CHECKED(UMLDeformerVizSettings, bDrawVertex));
+	TrainingMeshesCategoryBuilder.AddProperty(GET_MEMBER_NAME_CHECKED(UMLDeformerVizSettings, bDrawDeltas));
+	TrainingMeshesCategoryBuilder.AddProperty(GET_MEMBER_NAME_CHECKED(UMLDeformerVizSettings, bXRayDeltas));
+}
+
+bool FMLDeformerVizSettingsDetails::FilterAnimSequences(const FAssetData& AssetData, USkeleton* Skeleton)
+{
+	if (Skeleton && Skeleton->IsCompatibleSkeletonByAssetData(AssetData))
+	{
+		return false;
+	}
+
+	return true;
+}
+
+void FMLDeformerVizSettingsDetails::OnResetToDefaultDeformerGraph(TSharedPtr<IPropertyHandle> PropertyHandle)
+{
+	UMLDeformerAsset* DeformerAsset = GetMLDeformerAsset();
+	UMLDeformerVizSettings* VizSettings = DeformerAsset->GetVizSettings();
+
+	UComputeGraph* DefaultGraph = FMLDeformerEditorData::LoadDefaultDeformerGraph();
+	PropertyHandle->SetValue(DefaultGraph);
+}
+
+bool FMLDeformerVizSettingsDetails::IsResetToDefaultDeformerGraphVisible(TSharedPtr<IPropertyHandle> PropertyHandle)
+{
+	UObject* CurrentGraph = nullptr;
+	PropertyHandle->GetValue(CurrentGraph);
+	if (CurrentGraph == nullptr)
+	{
+		return true;
+	}
+
+	// Check if we already assigned the default asset.
+	const FAssetData CurrentGraphAssetData(CurrentGraph);
+	const FString CurrentPath = CurrentGraphAssetData.ObjectPath.ToString();
+	const FString DefaultPath = FMLDeformerEditorData::GetDefaultDeformerGraphAssetPath();
+	return (DefaultPath != CurrentPath);
+}
+
+#undef LOCTEXT_NAMESPACE
diff --git a/Engine/Plugins/Experimental/MLDeformer/Source/MLDeformerEditor/Private/MLDeformerVizSettingsDetails.h b/Engine/Plugins/Experimental/MLDeformer/Source/MLDeformerEditor/Private/MLDeformerVizSettingsDetails.h
new file mode 100644
--- /dev/null
+++ b/Engine/Plugins/Experimental/MLDeformer/Source/MLDeformerEditor/Private/MLDeformerVizSettingsDetails.h
@@ -0,0 +1,34 @@
+// Copyright Epic Games, Inc. All Rights Reserved.
+
+#pragma once
+
+#include "CoreMinimal.h"
+#include "IDetailCustomization.h"
+#include "AssetRegistry/AssetData.h"
+#include "PropertyHandle.h"
+
+class IDetailLayoutBuilder;
+class USkeleton;
+class UMLDeformerAsset;
+
+class FMLDeformerVizSettingsDetails : public IDetailCustomization
+{
+public:
+	/** Makes a new instance of this detail layout class for a specific detail view requesting it. */
+	static TSharedRef<IDetailCustomization> MakeInstance();
+
+	/** ILayoutDetails interface. */
+	virtual void CustomizeDetails(class IDetailLayoutBuilder& DetailBuilder) override;
+
+	IDetailLayoutBuilder* GetDetailLayoutBuilder() const { return DetailLayoutBuilder; }
+
+private:
+	bool FilterAnimSequences(const FAssetData& AssetData, USkeleton* Skeleton);
+	void OnResetToDefaultDeformerGraph(TSharedPtr<IPropertyHandle> PropertyHandle);
+	bool IsResetToDefaultDeformerGraphVisible(TSharedPtr<IPropertyHandle> PropertyHandle);
+	UMLDeformerAsset* GetMLDeformerAsset() const;
+
+private:
+	/** Associated detail layout builder. */
+	IDetailLayoutBuilder* DetailLayoutBuilder = nullptr;
+};
diff --git a/Engine/Plugins/Experimental/MLDeformer/Source/MLDeformerEditor/Private/MLDeformerVizSettingsTabSummoner.cpp b/Engine/Plugins/Experimental/MLDeformer/Source/MLDeformerEditor/Private/MLDeformerVizSettingsTabSummoner.cpp
new file mode 100644
--- /dev/null
+++ b/Engine/Plugins/Experimental/MLDeformer/Source/MLDeformerEditor/Private/MLDeformerVizSettingsTabSummoner.cpp
@@ -0,0 +1,62 @@
+// Copyright Epic Games, Inc. All Rights Reserved.
+
+#include "MLDeformerVizSettingsTabSummoner.h"
+#include "MLDeformerEditorToolkit.h"
+#include "MLDeformerEditorData.h"
+#include "MLDeformerEditorStyle.h"
+
+#include "IDocumentation.h"
+#include "SScrubControlPanel.h"
+
+#define LOCTEXT_NAMESPACE "MLDeformerVizSettingsTabSummoner"
+
+const FName FMLDeformerVizSettingsTabSummoner::TabID(TEXT("MLDeformerVizSettings"));
+
+FMLDeformerVizSettingsTabSummoner::FMLDeformerVizSettingsTabSummoner(const TSharedRef<FMLDeformerEditorToolkit>& InEditor)
+	: FWorkflowTabFactory(TabID, InEditor)
+	, Editor(InEditor)
+{
+	bIsSingleton = true; // only allow a single instance of this tab
+
+	TabLabel = LOCTEXT("VizSettingsTabLabel", "Visualization");
+	TabIcon = FSlateIcon(FMLDeformerEditorStyle::Get().GetStyleSetName(), "MLDeformer.VizSettings.TabIcon");
+	ViewMenuDescription = LOCTEXT("ViewMenu_Desc", "Visualization Settings");
+	ViewMenuTooltip = LOCTEXT("ViewMenu_ToolTip", "Show the ML Deformer Debug Visualization Settings");
+
+	// Create details view.
+	FDetailsViewArgs DetailsViewArgs;
+	DetailsViewArgs.NameAreaSettings = FDetailsViewArgs::HideNameArea;
+	DetailsViewArgs.bHideSelectionTip = true;
+
+	FPropertyEditorModule& PropertyEditorModule = FModuleManager::LoadModuleChecked<FPropertyEditorModule>("PropertyEditor");
+	DetailsView = PropertyEditorModule.CreateDetailView(DetailsViewArgs);
+	Editor.Pin()->GetEditorData()->SetVizSettingsDetailsView(DetailsView);
+}
+
+TSharedPtr<SToolTip> FMLDeformerVizSettingsTabSummoner::CreateTabToolTipWidget(const FWorkflowTabSpawnInfo& Info) const
+{
+	return IDocumentation::Get()->CreateToolTip(
+		LOCTEXT("VizSettingsTooltip", "The visualization settings for the ML Deformer."), 
+		NULL, 
+		TEXT("Shared/Editors/Persona"), 
+		TEXT("MLDeformerVizSettings_Window"));
+}
+
+TSharedRef<SWidget> FMLDeformerVizSettingsTabSummoner::CreateTabBody(const FWorkflowTabSpawnInfo& Info) const
+{
+	TSharedRef<SVerticalBox> Content = SNew(SVerticalBox);
+	Content->AddSlot()
+		.FillHeight(1.0f)
+		[
+			DetailsView.ToSharedRef()
+		];
+
+	// Init the viz settings details panel.
+	UMLDeformerVizSettings* VizSettings = Editor.Pin()->GetEditorData()->GetDeformerAsset()->VizSettings;
+	Editor.Pin()->GetEditorData()->GetVizSettingsDetailsView()->OnFinishedChangingProperties().AddSP(Editor.Pin().ToSharedRef(), &FMLDeformerEditorToolkit::OnFinishedChangingDetails);
+	Editor.Pin()->GetEditorData()->GetVizSettingsDetailsView()->SetObject(VizSettings);
+
+	return Content;
+}
+
+#undef LOCTEXT_NAMESPACE 
diff --git a/Engine/Plugins/Experimental/MLDeformer/Source/MLDeformerEditor/Private/MLDeformerVizSettingsTabSummoner.h b/Engine/Plugins/Experimental/MLDeformer/Source/MLDeformerEditor/Private/MLDeformerVizSettingsTabSummoner.h
new file mode 100644
--- /dev/null
+++ b/Engine/Plugins/Experimental/MLDeformer/Source/MLDeformerEditor/Private/MLDeformerVizSettingsTabSummoner.h
@@ -0,0 +1,23 @@
+// Copyright Epic Games, Inc. All Rights Reserved.
+
+#pragma once
+
+#include "CoreMinimal.h"
+#include "WorkflowOrientedApp/WorkflowTabFactory.h"
+
+class FMLDeformerEditorToolkit;
+class IDetailsView;
+
+struct FMLDeformerVizSettingsTabSummoner : public FWorkflowTabFactory
+{
+public:
+	static const FName TabID;
+
+	FMLDeformerVizSettingsTabSummoner(const TSharedRef<FMLDeformerEditorToolkit>& InEditor);
+	virtual TSharedRef<SWidget> CreateTabBody(const FWorkflowTabSpawnInfo& Info) const override;
+	virtual TSharedPtr<SToolTip> CreateTabToolTipWidget(const FWorkflowTabSpawnInfo& Info) const override;
+
+protected:
+	TWeakPtr<FMLDeformerEditorToolkit> Editor;
+	TSharedPtr<IDetailsView> DetailsView;
+};
diff --git a/Engine/Plugins/Experimental/MLDeformer/Source/MLDeformerEditor/Private/MLPytorchDataSetInterface.cpp b/Engine/Plugins/Experimental/MLDeformer/Source/MLDeformerEditor/Private/MLPytorchDataSetInterface.cpp
new file mode 100644
--- /dev/null
+++ b/Engine/Plugins/Experimental/MLDeformer/Source/MLDeformerEditor/Private/MLPytorchDataSetInterface.cpp
@@ -0,0 +1,101 @@
+﻿// Copyright Epic Games, Inc. All Rights Reserved.
+
+
+#include "MLPytorchDataSetInterface.h"
+
+#include "MLDeformer.h"
+#include "MLDeformerEditorData.h"
+#include "MLDeformerAsset.h"
+
+#include "Animation/DebugSkelMeshComponent.h"
+#include "GeometryCache.h"
+
+UMLPytorchDataSetInterface::UMLPytorchDataSetInterface()
+{
+}
+
+void UMLPytorchDataSetInterface::SetEditorData(TSharedPtr<FMLDeformerEditorData> InEditorData)
+{
+	EditorData = InEditorData;
+}
+
+bool UMLPytorchDataSetInterface::IsValid() const
+{
+	const FMLDeformerEditorData* Data = EditorData.Get();
+	return (Data != nullptr);
+}
+
+int32 UMLPytorchDataSetInterface::GetNumberSampleTransforms() const
+{
+	check(IsValid());
+	const FMLDeformerEditorData* Data = EditorData.Get();
+	return Data->GetDeformerAsset()->GetInputInfo().GetNumBones();
+}
+
+int32 UMLPytorchDataSetInterface::GetNumberSampleCurves() const
+{
+	check(IsValid());
+	const FMLDeformerEditorData* Data = EditorData.Get();
+	return Data->GetDeformerAsset()->GetInputInfo().GetNumCurves();
+}
+
+int32 UMLPytorchDataSetInterface::GetNumberSampleDeltas() const
+{
+	check(IsValid());
+	const FMLDeformerEditorData* Data = EditorData.Get();
+	return Data->GetNumImportedVertices();
+}
+
+int32 UMLPytorchDataSetInterface::NumSamples() const
+{
+	check(IsValid());
+	return EditorData->GetDeformerAsset()->GetNumFrames();
+}
+
+bool UMLPytorchDataSetInterface::SetCurrentSampleIndex(int32 Index)
+{
+	check(IsValid());
+
+	FMLDeformerEditorData* Data = EditorData.Get();
+	UMLDeformerAsset* DeformerAsset = Data->GetDeformerAsset();
+
+	// Make sure we have a valid frame number.
+	if (Index < 0 || Index >= DeformerAsset->GetNumFrames())
+	{
+		UE_LOG(LogMLDeformer, Warning, TEXT("Sample index must range from %d to %d, but a value of %d was provided."), 0, DeformerAsset->GetNumFrames()-1, Index);
+		return false;
+	}
+
+	// Generate and store the deltas.
+	const float DeltaCutoffLength = DeformerAsset->GetDeltaCutoffLength();
+	if (!Data->GenerateTrainingData(0, Index, DeltaCutoffLength, SampleDeltas, SampleBoneTransforms, SampleCurveWeights))
+	{
+		return false;
+	}
+
+	return true;
+}
+
+bool UMLPytorchDataSetInterface::ComputeDeltasStatistics()
+{
+	check(IsValid());
+
+	FMLDeformerEditorData* Data = EditorData.Get();
+	UMLDeformerAsset* DeformerAsset = Data->GetDeformerAsset();
+
+	// Generate and store the deltas.
+	const float DeltaCutoffLength = DeformerAsset->GetDeltaCutoffLength();
+	if (Data->ComputeVertexDeltaStatistics(0, DeltaCutoffLength))
+	{
+		//Update mean vertex delta and vertex data scale.
+		VertexDeltaMean.X = DeformerAsset->VertexDeltaMean.X;
+		VertexDeltaMean.Y = DeformerAsset->VertexDeltaMean.Y;
+		VertexDeltaMean.Z = DeformerAsset->VertexDeltaMean.Z;
+		VertexDeltaScale.X = DeformerAsset->VertexDeltaScale.X;
+		VertexDeltaScale.Y = DeformerAsset->VertexDeltaScale.Y;
+		VertexDeltaScale.Z = DeformerAsset->VertexDeltaScale.Z;
+		return true;
+	}
+
+	return false;
+}
diff --git a/Engine/Plugins/Experimental/MLDeformer/Source/MLDeformerEditor/Public/MLDeformerApplicationMode.h b/Engine/Plugins/Experimental/MLDeformer/Source/MLDeformerEditor/Public/MLDeformerApplicationMode.h
new file mode 100644
--- /dev/null
+++ b/Engine/Plugins/Experimental/MLDeformer/Source/MLDeformerEditor/Public/MLDeformerApplicationMode.h
@@ -0,0 +1,27 @@
+// Copyright Epic Games, Inc. All Rights Reserved.
+
+#pragma once
+
+#include "WorkflowOrientedApp/ApplicationMode.h"
+#include "WorkflowOrientedApp/WorkflowCentricApplication.h"
+#include "WorkflowOrientedApp/WorkflowTabManager.h"
+
+class FMLDeformerEditorToolkit;
+
+class FMLDeformerApplicationMode : public FApplicationMode
+{
+public:
+	static FName ModeName;
+
+	FMLDeformerApplicationMode(TSharedRef<class FWorkflowCentricApplication> InHostingApp, TSharedRef<class IPersonaPreviewScene> InPreviewScene);
+
+	/** FApplicationMode interface. */
+	virtual void RegisterTabFactories(TSharedPtr<FTabManager> InTabManager) override;
+
+protected:
+	/** The hosting app. */
+	TWeakPtr<FMLDeformerEditorToolkit> EditorToolkit;
+
+	/** The tab factories we support. */
+	FWorkflowAllowedTabSet TabFactories;
+};
diff --git a/Engine/Plugins/Experimental/MLDeformer/Source/MLDeformerEditor/Public/MLDeformerEditor.h b/Engine/Plugins/Experimental/MLDeformer/Source/MLDeformerEditor/Public/MLDeformerEditor.h
new file mode 100644
--- /dev/null
+++ b/Engine/Plugins/Experimental/MLDeformer/Source/MLDeformerEditor/Public/MLDeformerEditor.h
@@ -0,0 +1,18 @@
+// Copyright Epic Games, Inc. All Rights Reserved.
+
+#pragma once
+
+#include "CoreMinimal.h"
+#include "Modules/ModuleInterface.h"
+
+class FMLDeformerAssetActions;
+
+class MLDEFORMEREDITOR_API FMLDeformerEditor : public IModuleInterface
+{
+public:
+	void StartupModule() override;
+	void ShutdownModule() override;
+
+private:
+	TSharedPtr<FMLDeformerAssetActions> MLDeformerAssetActions;
+};
diff --git a/Engine/Plugins/Experimental/MLDeformer/Source/MLDeformerEditor/Public/MLDeformerEditorData.h b/Engine/Plugins/Experimental/MLDeformer/Source/MLDeformerEditor/Public/MLDeformerEditorData.h
new file mode 100644
--- /dev/null
+++ b/Engine/Plugins/Experimental/MLDeformer/Source/MLDeformerEditor/Public/MLDeformerEditorData.h
@@ -0,0 +1,174 @@
+// Copyright Epic Games, Inc. All Rights Reserved.
+
+#pragma once
+
+#include "CoreMinimal.h"
+#include "MLDeformerAsset.h"
+#include "MLDeformerInputInfo.h"
+
+class UAnimInstance;
+class UMLDeformerAsset;
+class UMLDeformerComponent;
+class UComputeGraphComponent;
+class UComputeGraph;
+class IPersonaToolkit;
+class UDebugSkelMeshComponent;
+class IDetailsView;
+class USkeletalMesh;
+class UGeometryCache;
+class UGeometryCacheComponent;
+class UMLDeformerVizSettings;
+class FMLDeformerEditorToolkit;
+class UMaterial;
+class AActor;
+class UTextRenderComponent;
+class SSimpleTimeSlider;
+struct FMLDeformerTrainingData;
+struct FReferenceSkeleton;
+struct FGeometryCacheMeshData;
+
+enum class EMLDeformerEditorActorIndex : int8
+{
+	Base = 0,		// Linear skinned.
+	Target,			// Geometry cache.
+	Test,			// Linear skinned test model.
+	DeformedTest,	// Test model with ML deformer applied to it.
+	GroundTruth		// Ground truth model that plays the same animation as the test anim asset (optional).
+};
+
+struct FMLDeformerEditorActor
+{
+	TObjectPtr<AActor> Actor = nullptr;
+	TObjectPtr<UTextRenderComponent> LabelComponent = nullptr;
+	TObjectPtr<UDebugSkelMeshComponent> SkelMeshComponent = nullptr;
+	TObjectPtr<UGeometryCacheComponent> GeomCacheComponent = nullptr;
+	TObjectPtr<UComputeGraphComponent> ComputeGraphComponent = nullptr;
+	TObjectPtr<UMLDeformerComponent> MLDeformerComponent = nullptr;
+};
+
+class FMLDeformerEditorData
+{
+public:
+	FMLDeformerEditorData();
+	~FMLDeformerEditorData() = default;
+
+	void SetPersonaToolkit(TSharedPtr<IPersonaToolkit> InPersonaToolkit);
+	void SetDeformerAsset(UMLDeformerAsset* InAsset);
+	void SetAnimInstance(UAnimInstance* InAnimInstance);
+	void SetDetailsView(TSharedPtr<IDetailsView> InDetailsView);
+	void SetEditorToolkit(FMLDeformerEditorToolkit* InToolkit);
+	void SetVizSettingsDetailsView(TSharedPtr<IDetailsView> InDetailsView);
+	void SetTimeSlider(TSharedPtr<SSimpleTimeSlider> InTimeSlider);
+
+	int32 GetNumEditorActors() const;
+	void SetEditorActor(EMLDeformerEditorActorIndex Index, const FMLDeformerEditorActor& Actor);
+	const FMLDeformerEditorActor& GetEditorActor(EMLDeformerEditorActorIndex Index) const;
+	FMLDeformerEditorActor& GetEditorActor(EMLDeformerEditorActorIndex Index);
+	bool IsTestActor(EMLDeformerEditorActorIndex Index) const;
+	bool IsTrainingActor(EMLDeformerEditorActorIndex Index) const;
+
+	void InitAssets();
+	bool GenerateDeltas(uint32 LODIndex, uint32 FrameNumber, float DeltaCutoffLength, TArray<float>& OutDeltas);
+	bool GenerateTrainingData(uint32 LODIndex, uint32 FrameNumber, float DeltaCutoffLength, TArray<float>& OutDeltas, TArray<FTransform>& OutBoneTransforms, TArray<float>& OutCurveValues);
+	bool ComputeVertexDeltaStatistics(uint32 LODIndex, float DeltaCutoffLength);
+
+	void SetTimeSliderRange(double StartTime, double EndTime);
+	void UpdateTimeSlider();
+	void OnTimeSliderScrubPositionChanged(double NewScrubTime, bool bIsScrubbing);
+	void OnPlayButtonPressed();
+	double CalcTimelinePosition() const;
+	bool IsPlayingAnim() const;
+
+	IPersonaToolkit* GetPersonaToolkit() const;
+	TSharedPtr<IPersonaToolkit> GetPersonaToolkitPointer() const;
+	UMLDeformerAsset* GetDeformerAsset() const;
+	TWeakObjectPtr<UMLDeformerAsset> GetDeformerAssetPointer() const;
+	UAnimInstance* GetAnimInstance() const;
+	IDetailsView* GetDetailsView() const;
+	IDetailsView* GetVizSettingsDetailsView() const;
+	FMLDeformerEditorToolkit* GetEditorToolkit() const;
+	SSimpleTimeSlider* GetTimeSlider() const;
+
+	float GetDuration() const;
+	float GetSnappedFrameTime(float InTime) const;
+	float GetTimeAtFrame(int32 FrameNumber) const;
+	int32 GetNumImportedVertices() const;
+
+	/**
+	 *Update mean and scale of vertex deltas.
+	 *
+	 * @param InGeomCachePositions Geometric cache positions
+	 * @param InLinearSkinnedPositions Skinned positions
+	 * @param InDeltaCutoffLength Delta cutoff length
+	 * @param InOutMeanVertexDelta Mean vertex delta
+	 * @param InOutVertexDeltaScale Vertex delta scale
+	 * @param InOutCount Count
+	 */
+	void UpdateVertexDeltaMeanAndScale(const TArray<FVector3f>& InGeomCachePositions, const TArray<FVector3f>& InLinearSkinnedPositions,
+		float InDeltaCutoffLength, FVector3f& InOutMeanVertexDelta, FVector3f& InOutVertexDeltaScale, float& InOutCount) const;
+
+	void ExtractBoneTransforms(TArray<FMatrix44f>& OutBoneMatrices, TArray<FTransform>& OutBoneTransforms) const;
+	void ExtractSkinnedPositions(int32 LODIndex, TArray<FMatrix44f>& InBoneMatrices, TArray<FVector3f>& TempPositions, TArray<FVector3f>& OutPositions, bool bImportedVertices=true) const;
+	void ExtractGeomCachePositions(int32 LODIndex, TArray<FVector3f>& TempPositions, TArray<FVector3f>& OutPositions) const;
+
+	void ExtractInputBoneTransforms(TArray<FTransform>& OutTransforms) const;
+	void ExtractInputCurveValues(TArray<float>& OutValues) const;
+
+	void UpdateDebugPointData();
+	void UpdateTestAnimPlaySpeed();
+	void UpdateDeformerGraph();
+	static FString GetDefaultDeformerGraphAssetPath();
+	static UComputeGraph* LoadDefaultDeformerGraph();
+	void SetAnimFrame(int32 FrameNumber);
+	void CreateMaterials();
+
+	bool IsReadyForTraining() const;
+	FText GetOverlayText();
+
+	void SetHeatMapMaterialEnabled(bool bEnabled);
+	void SetActorVisibility(EMLDeformerEditorActorIndex ActorIndex, bool bIsVisible);
+	void SetDefaultDeformerGraphIfNeeded();
+	bool IsActorVisible(EMLDeformerEditorActorIndex ActorIndex) const;
+
+public:
+	TArray<FMatrix44f> CurrentBoneMatrices;
+	TArray<FTransform> CurrentBoneTransforms;
+	TArray<float> CurrentCurveWeights;
+
+	TArray<FVector3f> LinearSkinnedPositions;
+	TArray<FVector3f> GeomCachePositions;
+	TArray<FVector3f> TempVectorBuffer;
+	TArray<float> VertexDeltas;
+	bool bIsVertexDeltaNormalized = false;
+	TArray<FGeometryCacheMeshData> MeshDatas;
+
+	int32 CurrentFrame = -1;
+
+private:
+	/** The editor actors. */
+	TArray<FMLDeformerEditorActor> Actors;
+
+	/** Preview scene to be supplied by IHasPersonaToolkit::GetPersonaToolkit. */
+	TSharedPtr<IPersonaToolkit> PersonaToolkit;
+
+	/** The ML Deformer asset we are editing. */
+	TWeakObjectPtr<UMLDeformerAsset> MLDeformerAsset;
+
+	/** Viewport anim instance, which is the anim instance playing on the linear skinned mesh. */
+	TObjectPtr<UAnimInstance> AnimInstance;
+
+	/** Asset details tab. */
+	TSharedPtr<IDetailsView> DetailsView;
+
+	/** Viz settings details tab. */
+	TSharedPtr<IDetailsView> VizSettingsDetailsView;
+
+	/** The editor toolkit. */
+	TObjectPtr<FMLDeformerEditorToolkit> EditorToolkit;
+
+	/** The heatmap material. */
+	TObjectPtr<UMaterial> HeatMapMaterial;
+
+	/** The timeline slider widget. */
+	TSharedPtr<SSimpleTimeSlider> TimeSlider;
+};
diff --git a/Engine/Plugins/Experimental/MLDeformer/Source/MLDeformerEditor/Public/MLDeformerEditorMode.h b/Engine/Plugins/Experimental/MLDeformer/Source/MLDeformerEditor/Public/MLDeformerEditorMode.h
new file mode 100644
--- /dev/null
+++ b/Engine/Plugins/Experimental/MLDeformer/Source/MLDeformerEditor/Public/MLDeformerEditorMode.h
@@ -0,0 +1,48 @@
+// Copyright Epic Games, Inc. All Rights Reserved.
+
+#pragma once
+
+#include "CoreMinimal.h"
+
+#include "MLDeformerAsset.h"
+#include "IPersonaEditMode.h"
+#include "MLDeformerEditorData.h"
+
+class FMLDeformerEditorToolkit;
+class FMLDeformerPreviewScene;
+class FMLDeformerEditorData;
+class UMLDeformerVizSettings;
+struct FMLDeformerEditorActor;
+
+class FMLDeformerEditorMode : public IPersonaEditMode
+{
+public:
+	static FName ModeName;
+
+	FMLDeformerEditorMode();
+
+	void SetEditorData(TSharedPtr<FMLDeformerEditorData> InEditorData);
+
+	/** IPersonaEditMode interface. */
+	virtual bool GetCameraTarget(FSphere& OutTarget) const override;
+	virtual class IPersonaPreviewScene& GetAnimPreviewScene() const override;
+	virtual void GetOnScreenDebugInfo(TArray<FText>& OutDebugInfo) const override;
+
+	/** FEdMode interface. */
+	virtual void Tick(FEditorViewportClient* ViewportClient, float DeltaTime) override;
+	virtual void Render(const FSceneView* View, FViewport* Viewport, FPrimitiveDrawInterface* PDI) override;
+	virtual void DrawHUD(FEditorViewportClient* ViewportClient, FViewport* Viewport, const FSceneView* View, FCanvas* Canvas) override;
+	virtual bool IsCompatibleWith(FEditorModeID OtherModeID) const override { return true; }
+	virtual bool AllowWidgetMove() override;
+	virtual bool ShouldDrawWidget() const override;
+	virtual bool UsesTransformWidget() const override;
+	virtual bool UsesTransformWidget(UE::Widget::EWidgetMode CheckMode) const override;
+
+private:
+	void UpdateActors();
+	void UpdateLabels();
+	void EncapsulateBounds(const FMLDeformerEditorActor* Actor, FBox& Box) const;
+
+private:
+	TWeakPtr<FMLDeformerEditorData> EditorData;
+};
diff --git a/Engine/Plugins/Experimental/MLDeformer/Source/MLDeformerEditor/Public/MLDeformerEditorToolkit.h b/Engine/Plugins/Experimental/MLDeformer/Source/MLDeformerEditor/Public/MLDeformerEditorToolkit.h
new file mode 100644
--- /dev/null
+++ b/Engine/Plugins/Experimental/MLDeformer/Source/MLDeformerEditor/Public/MLDeformerEditorToolkit.h
@@ -0,0 +1,105 @@
+// Copyright Epic Games, Inc. All Rights Reserved.
+
+#pragma once
+
+#include "CoreMinimal.h"
+#include "MLDeformerEditorData.h"
+
+#include "IHasPersonaToolkit.h"
+#include "IPersonaPreviewScene.h"
+#include "IPersonaViewport.h"
+#include "PersonaAssetEditorToolkit.h"
+#include "EditorUndoClient.h"
+#include "MLDeformerPythonTrainingModel.h"
+#include "Math/Color.h"
+
+#include "Widgets/Notifications/SNotificationList.h"
+
+class IDetailsView;
+class FGGAssetEditorToolbar;
+class UAnimInstance;
+class UMLDeformerAsset;
+class FMLDeformerEditorData;
+class UMLDeformerComponent;
+class UAnimPreviewInstance;
+class UTextRenderComponent;
+class UWorld;
+class UComputeGraph;
+
+namespace MLDeformerEditorModes
+{
+	extern const FName Editor;
+}
+
+class FMLDeformerEditorToolkit :
+	public FPersonaAssetEditorToolkit,
+	public IHasPersonaToolkit,
+	public FGCObject,
+	public FEditorUndoClient,
+	public FTickableEditorObject
+{
+public:
+	friend class FMLDeformerApplicationMode;
+	friend struct FMLDeformerVizSettingsTabSummoner;
+
+	FMLDeformerEditorToolkit();
+	virtual ~FMLDeformerEditorToolkit();
+
+	/** Initialize the asset editor. This will register the application mode, init the preview scene, etc. */
+	void InitAssetEditor(
+		const EToolkitMode::Type Mode,
+		const TSharedPtr<IToolkitHost>& InitToolkitHost,
+		UMLDeformerAsset* DeformerAsset);
+
+	/** FAssetEditorToolkit interface. */
+	virtual void RegisterTabSpawners(const TSharedRef<FTabManager>& InTabManager) override;
+	virtual void UnregisterTabSpawners(const TSharedRef<FTabManager>& InTabManager) override;
+	virtual FName GetToolkitFName() const override;
+	virtual FText GetBaseToolkitName() const override;
+	virtual FText GetToolkitName() const override;
+	virtual FLinearColor GetWorldCentricTabColorScale() const override;
+	virtual FString GetWorldCentricTabPrefix() const override;
+
+	/** FGCObject interface. */
+	virtual void AddReferencedObjects(FReferenceCollector& Collector) override;
+
+	/** FTickableEditorObject Interface. */
+	virtual void Tick(float DeltaTime) override {};
+	virtual ETickableTickType GetTickableTickType() const override { return ETickableTickType::Always; }
+	virtual TStatId GetStatId() const override;
+
+	/** IHasPersonaToolkit interface. */
+	virtual TSharedRef<IPersonaToolkit> GetPersonaToolkit() const override;
+
+	FMLDeformerEditorData* GetEditorData() const { return EditorData.Get(); }
+
+private:
+	/* Toolbar related. */
+	void ExtendToolbar();
+	void FillToolbar(FToolBarBuilder& ToolbarBuilder);
+
+	/** Preview scene setup. */
+	void HandlePreviewSceneCreated(const TSharedRef<IPersonaPreviewScene>& InPersonaPreviewScene);
+	void HandleViewportCreated(const TSharedRef<IPersonaViewport>& InPersonaViewport);
+	void HandleDetailsCreated(const TSharedRef<class IDetailsView>& InDetailsView);
+	void OnFinishedChangingDetails(const FPropertyChangedEvent& PropertyChangedEvent);
+
+	/** Helpers. */
+	void CreateSkinnedActor(EMLDeformerEditorActorIndex ActorIndex, const FName& Name, UWorld* World, USkeletalMesh* Mesh, const FLinearColor LabelColor, FLinearColor WireframeColor) const;
+	void CreateBaseActor(const TSharedRef<IPersonaPreviewScene>& InPersonaPreviewScene, const FName& Name, FLinearColor LabelColor, FLinearColor WireframeColor);
+	void CreateGeomCacheActor(EMLDeformerEditorActorIndex ActorIndex, UWorld* World, const FName& Name, UGeometryCache* GeomCache, FLinearColor LabelColor, FLinearColor WireframeColor);
+	UMLDeformerComponent* AddMLDeformerComponentToActor(EMLDeformerEditorActorIndex ActorIndex);
+	void AddDeformerGraphComponentToActor(EMLDeformerEditorActorIndex ActorIndex, UComputeGraph* ComputeGraph, UMLDeformerComponent* DeformerComponent);
+	bool TryLoadOnnxFile() const;
+	void ShowNotification(const FText& Message, SNotificationItem::ECompletionState State, bool PlaySound) const;
+	UTextRenderComponent* CreateLabelForActor(AActor* Actor, UWorld* World, FLinearColor Color, const FText& Text) const;
+	FText GetOverlayText() const;
+	void SetComputeGraphDataProviders() const;
+	void OnSwitchedVisualizationMode();
+	bool HandleTrainingResult(ETrainingResult TrainingResult, double TrainingDuration);
+	void UpdateActorVisibility();
+
+private:
+	/** The editor data, containing things like the skeletal mesh, anim instance, etc. */
+	TSharedPtr<FMLDeformerEditorData> EditorData;
+};
diff --git a/Engine/Plugins/Experimental/MLDeformer/Source/MLDeformerEditor/Public/MLPytorchDataSetInterface.h b/Engine/Plugins/Experimental/MLDeformer/Source/MLDeformerEditor/Public/MLPytorchDataSetInterface.h
new file mode 100644
--- /dev/null
+++ b/Engine/Plugins/Experimental/MLDeformer/Source/MLDeformerEditor/Public/MLPytorchDataSetInterface.h
@@ -0,0 +1,74 @@
+// Copyright Epic Games, Inc. All Rights Reserved.
+
+#pragma once
+
+#include "CoreMinimal.h"
+#include "UObject/Object.h"
+#include "MLPytorchDataSetInterface.generated.h"
+
+class FMLDeformerEditorData;
+/**
+ * 
+ */
+UCLASS(BlueprintType, hidecategories=Object)
+class MLDEFORMEREDITOR_API UMLPytorchDataSetInterface : public UObject
+{
+	GENERATED_BODY()
+public:
+	UMLPytorchDataSetInterface(); 
+	void SetEditorData(TSharedPtr<FMLDeformerEditorData> InEditorData);
+
+	/** Get the number of input transforms. */
+	UFUNCTION(BlueprintPure, Category = "Training Data")
+	int32 GetNumberSampleTransforms() const;
+
+	/** Get number of input curves. */
+	UFUNCTION(BlueprintPure, Category = "Training Data")
+	int32 GetNumberSampleCurves() const;
+
+	/** Get the number of output deltas. */
+	UFUNCTION(BlueprintPure, Category = "Training Data")
+	int32 GetNumberSampleDeltas() const;
+
+	/** Get the number of samples in this data set. */
+	UFUNCTION(BlueprintPure, Category = "Training Data")
+    int32 NumSamples() const;
+
+	/** Set the current sample index. This must be in range of [0..GetNumSamples()-1]. */
+	UFUNCTION(BlueprintCallable, Category = "Training Data")
+    bool SetCurrentSampleIndex(int32 Index);
+
+	/** Compute delta statistics for the whole dataset. */
+	UFUNCTION(BlueprintCallable, Category = "Training Data")
+	bool ComputeDeltasStatistics();
+
+	/** The delta values per vertex for this sample. This is updated after SetCurrentSampleIndex is called. */
+	UPROPERTY(EditAnywhere, BlueprintReadWrite, Category = "Training Data")
+	TArray<float> SampleDeltas;
+
+	/** The curve weights. This is updated after SetCurrentSampleIndex is called. */
+	UPROPERTY(EditAnywhere, BlueprintReadWrite, Category = "Training Data")
+	TArray<float> SampleCurveWeights;
+
+	/** Mean delta computed over the entire dataset. */
+	UPROPERTY(EditAnywhere, BlueprintReadWrite, Category = "Training Data")
+	FVector VertexDeltaMean = FVector::ZeroVector;
+
+	/** Vertex delta scale computed over the entire dataset. */
+	UPROPERTY(EditAnywhere, BlueprintReadWrite, Category = "Training Data")
+	FVector VertexDeltaScale = FVector::OneVector;
+
+	/** The bone transformations in component space for this sample. This is updated after SetCurrentSampleIndex is called. */
+	UPROPERTY(EditAnywhere, BlueprintReadWrite, Category = "Training Data")
+	TArray<FTransform> SampleBoneTransforms;
+
+	/** The current sample index. */
+	UPROPERTY(EditAnywhere, BlueprintReadOnly, Category = "Training Data")
+	int32 CurrentSampleIndex = -1;
+
+private:
+	bool IsValid() const;
+
+private:
+	TSharedPtr<FMLDeformerEditorData> EditorData;
+};
