MLDeformer: Removed nni back end check, both back ends are now compatible with MLDef.

[REDACTED_HANDLE], this is a trivial change that removes a check I added in the past (as now NNI would work with and without that check), sending it to you mostly for visibility.

#rnx

diff --git a/Engine/Plugins/Experimental/MLDeformer/Source/MLDeformer/Private/MLDeformerInstance.cpp b/Engine/Plugins/Experimental/MLDeformer/Source/MLDeformer/Private/MLDeformerInstance.cpp
--- a/Engine/Plugins/Experimental/MLDeformer/Source/MLDeformer/Private/MLDeformerInstance.cpp
+++ b/Engine/Plugins/Experimental/MLDeformer/Source/MLDeformer/Private/MLDeformerInstance.cpp
@@ -146,7 +146,6 @@ void FMLDeformerInstance::Update()
 	}
 
 	// We only support GPU processing of the neural network at the moment.
-	check(NeuralNetwork->GetBackEnd() == ENeuralBackEnd::UEOnly);
 	check(NeuralNetwork->GetInputDeviceType() == ENeuralDeviceType::CPU);
 	check(NeuralNetwork->GetDeviceType() == ENeuralDeviceType::GPU);
 	check(NeuralNetwork->GetOutputDeviceType() == ENeuralDeviceType::GPU);
