MLDeformer: Switched to ORT-based NNI.
#rnx

diff --git a/Engine/Plugins/Experimental/MLDeformer/Source/MLDeformer/Private/MLDeformerAsset.cpp b/Engine/Plugins/Experimental/MLDeformer/Source/MLDeformer/Private/MLDeformerAsset.cpp
--- a/Engine/Plugins/Experimental/MLDeformer/Source/MLDeformer/Private/MLDeformerAsset.cpp
+++ b/Engine/Plugins/Experimental/MLDeformer/Source/MLDeformer/Private/MLDeformerAsset.cpp
@@ -64,7 +64,7 @@ void UMLDeformerAsset::PostLoad()
 	Super::PostLoad();
 	if (NeuralNetwork)
 	{
-		NeuralNetwork->SetBackEnd(ENeuralBackEnd::UEOnly);
+		NeuralNetwork->SetBackEnd(ENeuralBackEnd::Auto);
 		NeuralNetwork->SetInputDeviceType(ENeuralDeviceType::CPU);
 		NeuralNetwork->SetDeviceType(ENeuralDeviceType::GPU);
 		NeuralNetwork->SetOutputDeviceType(ENeuralDeviceType::GPU);
diff --git a/Engine/Plugins/Experimental/MLDeformer/Source/MLDeformerEditor/Private/MLDeformerEditorToolkit.cpp b/Engine/Plugins/Experimental/MLDeformer/Source/MLDeformerEditor/Private/MLDeformerEditorToolkit.cpp
--- a/Engine/Plugins/Experimental/MLDeformer/Source/MLDeformerEditor/Private/MLDeformerEditorToolkit.cpp
+++ b/Engine/Plugins/Experimental/MLDeformer/Source/MLDeformerEditor/Private/MLDeformerEditorToolkit.cpp
@@ -336,7 +336,7 @@ bool FMLDeformerEditorToolkit::TryLoadOnnxFile() const
 		if (Network->Load(OnnxFile))
 		{
 			EditorData->GetDeformerAsset()->NeuralNetwork = Network;
-			EditorData->GetDeformerAsset()->NeuralNetwork->SetBackEnd(ENeuralBackEnd::UEOnly);
+			EditorData->GetDeformerAsset()->NeuralNetwork->SetBackEnd(ENeuralBackEnd::Auto);
 			EditorData->GetDeformerAsset()->NeuralNetwork->SetInputDeviceType(ENeuralDeviceType::CPU);
 			EditorData->GetDeformerAsset()->NeuralNetwork->SetDeviceType(ENeuralDeviceType::GPU);
 			EditorData->GetDeformerAsset()->NeuralNetwork->SetOutputDeviceType(ENeuralDeviceType::GPU);
