[MLDeformer]
- Show a dialog with progress when preprocessing the data.
- Allow cancelling this preprocessing.
- Force sampling raw animation data for the training data.

CR: https://p4-swarm.epicgames.net/reviews/17644865

diff --git a/Engine/Plugins/Experimental/MLDeformer/Source/MLDeformer/Public/MLDeformerAsset.h b/Engine/Plugins/Experimental/MLDeformer/Source/MLDeformer/Public/MLDeformerAsset.h
--- a/Engine/Plugins/Experimental/MLDeformer/Source/MLDeformer/Public/MLDeformerAsset.h
+++ b/Engine/Plugins/Experimental/MLDeformer/Source/MLDeformer/Public/MLDeformerAsset.h
@@ -198,7 +198,10 @@ public:
 	UPROPERTY(EditAnywhere, Category = "Target Mesh")
 	TObjectPtr<UGeometryCache> GeometryCache = nullptr;
 
-	/** The animation sequence to apply to the base mesh. This has to match the animation of the target mesh's geometry cache. */
+	/**
+	 * The animation sequence to apply to the base mesh. This has to match the animation of the target mesh's geometry cache. 
+	 * Internally we force the Interpolation property for this motion to be "Step".
+	 */
 	UPROPERTY(EditAnywhere, Category = "Base Mesh")
 	TObjectPtr<UAnimSequence> AnimSequence = nullptr;
 
diff --git a/Engine/Plugins/Experimental/MLDeformer/Source/MLDeformer/Public/MLDeformerComponent.h b/Engine/Plugins/Experimental/MLDeformer/Source/MLDeformer/Public/MLDeformerComponent.h
--- a/Engine/Plugins/Experimental/MLDeformer/Source/MLDeformer/Public/MLDeformerComponent.h
+++ b/Engine/Plugins/Experimental/MLDeformer/Source/MLDeformer/Public/MLDeformerComponent.h
@@ -47,7 +47,7 @@ public:
 private:
 	/** The deformer asset to use. */
 	UPROPERTY(EditAnywhere, Category = "ML Deformer")
-	UMLDeformerAsset* DeformerAsset = nullptr;
+	TObjectPtr<UMLDeformerAsset> DeformerAsset = nullptr;
 
 	/** The scale factor of the vertex deltas that are being applied. */
 	UPROPERTY(EditAnywhere, Category = "ML Deformer")
@@ -58,7 +58,7 @@ private:
 	 * This can be a nullptr. When it is a nullptr then it will internally try to find the first skeletal mesh component on the actor.
 	 * You can see this as an override. You can specify this override through the SetupComponent function.
 	 */
-	USkeletalMeshComponent* SkelMeshComponent = nullptr;
+	TObjectPtr<USkeletalMeshComponent> SkelMeshComponent = nullptr;
 
 	/** The mesh deformer instance. This is used to perform the runtime updates and run the inference. */
 	FMLDeformerInstance DeformerInstance;
diff --git a/Engine/Plugins/Experimental/MLDeformer/Source/MLDeformerEditor/Private/MLDeformerApplicationMode.cpp b/Engine/Plugins/Experimental/MLDeformer/Source/MLDeformerEditor/Private/MLDeformerApplicationMode.cpp
--- a/Engine/Plugins/Experimental/MLDeformer/Source/MLDeformerEditor/Private/MLDeformerApplicationMode.cpp
+++ b/Engine/Plugins/Experimental/MLDeformer/Source/MLDeformerEditor/Private/MLDeformerApplicationMode.cpp
@@ -27,13 +27,12 @@ FMLDeformerApplicationMode::FMLDeformerApplicationMode(
 	TSharedRef<FMLDeformerEditorToolkit> MLDeformerEditor = StaticCastSharedRef<FMLDeformerEditorToolkit>(InHostingApp);
 
 	FPersonaViewportArgs ViewportArgs(InPreviewScene);
-	ViewportArgs.bAlwaysShowTransformToolbar = false;
+	ViewportArgs.bAlwaysShowTransformToolbar = true;
 	ViewportArgs.bShowStats = false;
-	ViewportArgs.bShowTurnTable = false;
 	ViewportArgs.bShowTimeline = false;
-	ViewportArgs.bShowPhysicsMenu = false;
 	ViewportArgs.bShowLODMenu = true;
 	ViewportArgs.bShowPlaySpeedMenu = false;
+	ViewportArgs.bShowPhysicsMenu = false;
 	ViewportArgs.ContextName = TEXT("MLDeformerEditor.Viewport");
 	ViewportArgs.OnViewportCreated = FOnViewportCreated::CreateSP(MLDeformerEditor, &FMLDeformerEditorToolkit::HandleViewportCreated);
 
diff --git a/Engine/Plugins/Experimental/MLDeformer/Source/MLDeformerEditor/Private/MLDeformerEditorData.cpp b/Engine/Plugins/Experimental/MLDeformer/Source/MLDeformerEditor/Private/MLDeformerEditorData.cpp
--- a/Engine/Plugins/Experimental/MLDeformer/Source/MLDeformerEditor/Private/MLDeformerEditorData.cpp
+++ b/Engine/Plugins/Experimental/MLDeformer/Source/MLDeformerEditor/Private/MLDeformerEditorData.cpp
@@ -26,6 +26,7 @@
 #include "GeometryCacheMeshData.h"
 
 #include "Math/NumericLimits.h"
+#include "Misc/SlowTask.h"
 #include "ComputeFramework/ComputeGraphComponent.h"
 #include "ComputeFramework/ComputeGraph.h"
 
@@ -250,6 +251,10 @@ void FMLDeformerEditorData::UpdateDeformerGraph()
 
 void FMLDeformerEditorData::InitAssets()
 {
+	// Force the training sequence to use Step interpolation and sample raw animation data.
+	UAnimSequence* TrainingAnimSequence = MLDeformerAsset->GetAnimSequence();
+	TrainingAnimSequence->bUseRawDataOnly = true;
+
 	USkeletalMeshComponent* SkeletalMeshComponent = GetEditorActor(EMLDeformerEditorActorIndex::Base).SkelMeshComponent;
 	SkeletalMeshComponent->SetSkeletalMesh(MLDeformerAsset->GetSkeletalMesh());
 	if (PersonaToolkit)
@@ -257,7 +262,7 @@ void FMLDeformerEditorData::InitAssets()
 		PersonaToolkit->GetPreviewScene()->SetPreviewMesh(MLDeformerAsset->GetSkeletalMesh());
 	}
 	SkeletalMeshComponent->SetAnimationMode(EAnimationMode::AnimationSingleNode);
-	SkeletalMeshComponent->SetAnimation(MLDeformerAsset->GetAnimSequence());
+	SkeletalMeshComponent->SetAnimation(TrainingAnimSequence);
 	SkeletalMeshComponent->SetPosition(0.0f);
 	SkeletalMeshComponent->SetPlayRate(0.0f);
 	SkeletalMeshComponent->Play(false);
@@ -597,17 +602,21 @@ bool FMLDeformerEditorData::ComputeVertexDeltaStatistics(uint32 LODIndex, float
 {
 	if (!IsReadyForTraining() || bIsVertexDeltaNormalized)
 	{
-		return false;
+		return true;
 	}
 
 	// Calculate the number of frames to output.
 	UGeometryCacheComponent* GeometryCacheComponent = GetEditorActor(EMLDeformerEditorActorIndex::Target).GeomCacheComponent;
 	const uint32 AnimNumFrames = static_cast<uint32>(GeometryCacheComponent->GetNumberOfFrames());
 
+	// Show some dialog with progress.
+	FSlowTask Task((float)AnimNumFrames, FText::FromString("Calculating data statistics"));
+	Task.Initialize();
+	Task.MakeDialog(true);
+
 	// Initialize mean vertex delta and vertex delta scale.
-	MLDeformerAsset->VertexDeltaMean = FVector3f::ZeroVector;
-	MLDeformerAsset->VertexDeltaScale = FVector3f::OneVector;
-	FVector3f TempVertexDeltaScale = FVector3f::ZeroVector;
+	FVector3f VertexDeltaMean = FVector3f::ZeroVector;
+	FVector3f VertexDeltaScale = FVector3f::OneVector;
 	float Count = 0.0f;
 	TArray<FTransform> BoneTransforms;
 	for (uint32 FrameIndex = 0; FrameIndex < AnimNumFrames; FrameIndex++)
@@ -626,12 +635,22 @@ bool FMLDeformerEditorData::ComputeVertexDeltaStatistics(uint32 LODIndex, float
 
 		// Update mean vertex delta.
 		UpdateVertexDeltaMeanAndScale(GeomCachePositions, LinearSkinnedPositions, DeltaCutoffLength,
-			MLDeformerAsset->VertexDeltaMean, TempVertexDeltaScale, Count);
+			VertexDeltaMean, VertexDeltaScale, Count);
+
+		// Forward the progress bar and check if we want to cancel.
+		Task.EnterProgressFrame();
+		if (Task.ShouldCancel())
+		{
+			Task.Destroy();
+			return false;
+		}
 	}
 
+	// Update the asset with calculated statistics.
+	MLDeformerAsset->VertexDeltaMean = VertexDeltaMean;
 	if (Count > 0.0f)
 	{
-		MLDeformerAsset->VertexDeltaScale *= TempVertexDeltaScale.GetMax();
+		MLDeformerAsset->VertexDeltaScale = FVector3f::OneVector * VertexDeltaScale.GetMax();
 		bIsVertexDeltaNormalized = true;
 	}
 
diff --git a/Engine/Plugins/Experimental/MLDeformer/Source/MLDeformerEditor/Private/MLDeformerEditorToolkit.cpp b/Engine/Plugins/Experimental/MLDeformer/Source/MLDeformerEditor/Private/MLDeformerEditorToolkit.cpp
--- a/Engine/Plugins/Experimental/MLDeformer/Source/MLDeformerEditor/Private/MLDeformerEditorToolkit.cpp
+++ b/Engine/Plugins/Experimental/MLDeformer/Source/MLDeformerEditor/Private/MLDeformerEditorToolkit.cpp
@@ -296,6 +296,14 @@ bool FMLDeformerEditorToolkit::HandleTrainingResult(ETrainingResult TrainingResu
 		}
 		break;
 
+		// Training aborted but we cannot use the current network.
+		case ETrainingResult::AbortedCantUse:
+		{
+			ShowNotification(LOCTEXT("TrainingAborted", "Training aborted!"), SNotificationItem::ECompletionState::CS_None, true);
+			WindowMessage = LOCTEXT("TrainingAbortedCantUse", "Training aborted by user.");
+		}
+		break;
+
 		// Training data had issues.
 		case ETrainingResult::FailOnData:
 		{
diff --git a/Engine/Plugins/Experimental/MLDeformer/Source/MLDeformerEditor/Private/MLDeformerPythonTrainingModel.h b/Engine/Plugins/Experimental/MLDeformer/Source/MLDeformerEditor/Private/MLDeformerPythonTrainingModel.h
--- a/Engine/Plugins/Experimental/MLDeformer/Source/MLDeformerEditor/Private/MLDeformerPythonTrainingModel.h
+++ b/Engine/Plugins/Experimental/MLDeformer/Source/MLDeformerEditor/Private/MLDeformerPythonTrainingModel.h
@@ -16,6 +16,7 @@ enum class ETrainingResult : uint8
 {
 	Success = 0,	/** The training successfully finished. */
 	Aborted,		/** The user has aborted the training process. */
+	AbortedCantUse,	/** The user has aborted the training process and we can't use the resulting network. */
 	FailOnData,		/** The input or output data to the network has issues, which means we cannot train. */
 	FailUnknown		/** There is an unknown error (see output log). */
 };
