[MLDeformer]
- Introduce a memory based training data cache to speed up training.
- Refactors to make use of this cache and not have multiple calculations of the same data.
- Skip normalization calculations when the normalization factors already have been calculated.

diff --git a/Engine/Plugins/Experimental/MLDeformer/Source/MLDeformer/Private/MLDeformerInputInfo.cpp b/Engine/Plugins/Experimental/MLDeformer/Source/MLDeformer/Private/MLDeformerInputInfo.cpp
--- a/Engine/Plugins/Experimental/MLDeformer/Source/MLDeformer/Private/MLDeformerInputInfo.cpp
+++ b/Engine/Plugins/Experimental/MLDeformer/Source/MLDeformer/Private/MLDeformerInputInfo.cpp
@@ -233,10 +233,11 @@ FString FMLDeformerInputInfo::GenerateCompatibilityErrorString(USkeletalMesh* Sk
 	return ErrorString;
 }
 
-void FMLDeformerInputInfo::ExtractCurveValues(UAnimInstance* AnimInstance, TArray<float>& OutValues) const
+void FMLDeformerInputInfo::ExtractCurveValues(USkeletalMeshComponent* SkelMeshComponent, TArray<float>& OutValues) const
 {
 	check(CurveNames.Num() == CurveNameStrings.Num());
-	
+
+	UAnimInstance* AnimInstance = SkelMeshComponent->GetAnimInstance();
 	const int32 NumCurves = CurveNames.Num();
 	OutValues.Reset(NumCurves);
 	OutValues.AddUninitialized(NumCurves);
@@ -247,6 +248,26 @@ void FMLDeformerInputInfo::ExtractCurveValues(UAnimInstance* AnimInstance, TArra
 	}
 }
 
+void FMLDeformerInputInfo::ExtractBoneRotations(USkeletalMeshComponent* SkelMeshComponent, TArray<float>& OutRotations) const
+{
+	const TArray<FTransform>& BoneTransforms = SkelMeshComponent->GetBoneSpaceTransforms();
+	const int32 NumBones = GetNumBones();
+	OutRotations.Reset(NumBones * 4); // xyzw
+	OutRotations.AddUninitialized(NumBones * 4);
+	for (int32 Index = 0; Index < NumBones; ++Index)
+	{
+		const FName BoneName = GetBoneName(Index);
+		const int32 SkelMeshBoneIndex = SkelMeshComponent->GetBoneIndex(BoneName);
+		check(SkelMeshBoneIndex != INDEX_NONE);
+		const int32 Offset = Index * 4;
+		const FQuat Rotation = BoneTransforms[SkelMeshBoneIndex].GetRotation();
+		OutRotations[Offset] = Rotation.X;
+		OutRotations[Offset+1] = Rotation.Y;
+		OutRotations[Offset+2] = Rotation.Z;
+		OutRotations[Offset+3] = Rotation.W;
+	}
+}
+
 int32 FMLDeformerInputInfo::CalcNumNeuralNetInputs() const
 {
 	return 
diff --git a/Engine/Plugins/Experimental/MLDeformer/Source/MLDeformer/Public/MLDeformerAsset.h b/Engine/Plugins/Experimental/MLDeformer/Source/MLDeformer/Public/MLDeformerAsset.h
--- a/Engine/Plugins/Experimental/MLDeformer/Source/MLDeformer/Public/MLDeformerAsset.h
+++ b/Engine/Plugins/Experimental/MLDeformer/Source/MLDeformer/Public/MLDeformerAsset.h
@@ -143,6 +143,7 @@ public:
 	UMLDeformerVizSettings* GetVizSettings() { return VizSettings; }
 
 	int32 GetNumFrames() const;
+	int32 GetCacheSizeInMegabytes() const { return CacheSizeInMegabytes; }
 	float GetDeltaCutoffLength() const { return DeltaCutoffLength; }
 #endif
 
@@ -229,6 +230,10 @@ public:
 	UPROPERTY(EditAnywhere, BlueprintReadWrite, Category = "Training Settings", meta = (ClampMin = "1"))
 	int32 EpochsWithDecay = 10;
 
+	/** The maximum allowed size of the training cache in memory, in megabytes. So a value of 1024 would be one gigabyte. The larger the cache size the faster the training. */
+	UPROPERTY(EditAnywhere, BlueprintReadWrite, Category = "Training Settings", meta = (ClampMin = 0))
+	int32 CacheSizeInMegabytes = 2048;
+
 	/** The learning rate used during the model training. */
 	UPROPERTY(EditAnywhere, BlueprintReadWrite, AdvancedDisplay, Category = "Training Settings", meta = (ClampMin = "0.000001", ClampMax = "1.0"))
 	float LearningRate = 0.0005f;
diff --git a/Engine/Plugins/Experimental/MLDeformer/Source/MLDeformer/Public/MLDeformerInputInfo.h b/Engine/Plugins/Experimental/MLDeformer/Source/MLDeformer/Public/MLDeformerInputInfo.h
--- a/Engine/Plugins/Experimental/MLDeformer/Source/MLDeformer/Public/MLDeformerInputInfo.h
+++ b/Engine/Plugins/Experimental/MLDeformer/Source/MLDeformer/Public/MLDeformerInputInfo.h
@@ -114,10 +114,19 @@ public:
 
 	/** 
 	 * Extract the curve values for all curves we're interested in.
-	 * @param AnimInstance The anim instance we are grabbing the curve values from.
+	 * @param SkelMeshComponent The skeletal mesh component to sample from.
 	 * @param OutValues The array to write the values to. This array will be reset/resized by this method.
 	 */
-	void ExtractCurveValues(UAnimInstance* AnimInstance, TArray<float>& OutValues) const;
+	void ExtractCurveValues(USkeletalMeshComponent* SkelMeshComponent, TArray<float>& OutValues) const;
+
+	/**
+	 * Extract bone space rotations, as a float array.
+	 * The number of output rotations are NumBones * 4, where the array contains values like: xyzw,xyzw,xyzw,... where each xyzw are the 
+	 * components of the bone's bone space (local space) rotation quaternion.
+	 * @param SkelMeshComponent The skeletal mesh component to sample from.
+	 * @param OutRotations The output rotation values. This array will be resized internally.
+	 */
+	void ExtractBoneRotations(USkeletalMeshComponent* SkelMeshComponent, TArray<float>& OutRotations) const;
 
 	/**
 	 * Get the number of imported vertices in the base mesh.
diff --git a/Engine/Plugins/Experimental/MLDeformer/Source/MLDeformerEditor/Private/MLDeformerAssetDetails.cpp b/Engine/Plugins/Experimental/MLDeformer/Source/MLDeformerEditor/Private/MLDeformerAssetDetails.cpp
--- a/Engine/Plugins/Experimental/MLDeformer/Source/MLDeformerEditor/Private/MLDeformerAssetDetails.cpp
+++ b/Engine/Plugins/Experimental/MLDeformer/Source/MLDeformerEditor/Private/MLDeformerAssetDetails.cpp
@@ -192,6 +192,7 @@ void FMLDeformerAssetDetails::CustomizeDetails(class IDetailLayoutBuilder& Detai
 	SettingsCategoryBuilder.AddProperty(GET_MEMBER_NAME_CHECKED(UMLDeformerAsset, BatchSize));
 	SettingsCategoryBuilder.AddProperty(GET_MEMBER_NAME_CHECKED(UMLDeformerAsset, Epochs));
 	SettingsCategoryBuilder.AddProperty(GET_MEMBER_NAME_CHECKED(UMLDeformerAsset, EpochsWithDecay));
+	SettingsCategoryBuilder.AddProperty(GET_MEMBER_NAME_CHECKED(UMLDeformerAsset, CacheSizeInMegabytes));
 
 	// Show a warning when no neural network has been set.
 	if (DeformerAsset)
diff --git a/Engine/Plugins/Experimental/MLDeformer/Source/MLDeformerEditor/Private/MLDeformerEditorData.cpp b/Engine/Plugins/Experimental/MLDeformer/Source/MLDeformerEditor/Private/MLDeformerEditorData.cpp
--- a/Engine/Plugins/Experimental/MLDeformer/Source/MLDeformerEditor/Private/MLDeformerEditorData.cpp
+++ b/Engine/Plugins/Experimental/MLDeformer/Source/MLDeformerEditor/Private/MLDeformerEditorData.cpp
@@ -7,6 +7,7 @@
 #include "MLDeformerInputInfo.h"
 #include "MLDeformerEditorToolkit.h"
 #include "MLDeformer.h"
+#include "MLDeformerFrameCache.h"
 
 #include "IPersonaToolkit.h"
 #include "IDetailsView.h"
@@ -26,7 +27,7 @@
 #include "GeometryCacheMeshData.h"
 
 #include "Math/NumericLimits.h"
-#include "Misc/SlowTask.h"
+#include "Misc/ScopedSlowTask.h"
 #include "ComputeFramework/ComputeGraphComponent.h"
 #include "ComputeFramework/ComputeGraph.h"
 
@@ -253,7 +254,11 @@ void FMLDeformerEditorData::InitAssets()
 {
 	// Force the training sequence to use Step interpolation and sample raw animation data.
 	UAnimSequence* TrainingAnimSequence = MLDeformerAsset->GetAnimSequence();
-	TrainingAnimSequence->bUseRawDataOnly = true;
+	if (TrainingAnimSequence)
+	{
+		TrainingAnimSequence->bUseRawDataOnly = true;
+		TrainingAnimSequence->Interpolation = EAnimInterpolationType::Step;
+	}
 
 	USkeletalMeshComponent* SkeletalMeshComponent = GetEditorActor(EMLDeformerEditorActorIndex::Base).SkelMeshComponent;
 	SkeletalMeshComponent->SetSkeletalMesh(MLDeformerAsset->GetSkeletalMesh());
@@ -308,105 +313,14 @@ void FMLDeformerEditorData::InitAssets()
 	UpdateTimeSlider();
 	MLDeformerAsset->UpdateCachedNumVertices();
 	UpdateDeformerGraph();
-}
-
-void FMLDeformerEditorData::ExtractBoneTransforms(TArray<FMatrix44f>& OutBoneMatrices, TArray<FTransform>& OutBoneTransforms) const
-{
-	USkeletalMeshComponent* SkelComp = GetEditorActor(EMLDeformerEditorActorIndex::Base).SkelMeshComponent;
-	check(SkelComp);
 
-	SkelComp->CacheRefToLocalMatrices(OutBoneMatrices);
-	OutBoneTransforms = SkelComp->GetBoneSpaceTransforms();
-}
-
-void FMLDeformerEditorData::ExtractSkinnedPositions(int32 LODIndex, TArray<FMatrix44f>& InBoneMatrices, TArray<FVector3f>& TempPositions, TArray<FVector3f>& OutPositions, bool bImportedVertices) const
-{
-	OutPositions.Reset();
-	TempPositions.Reset();
-
-	UDebugSkelMeshComponent* SkelMeshComponent = GetEditorActor(EMLDeformerEditorActorIndex::Base).SkelMeshComponent;
-	if (SkelMeshComponent == nullptr)
-	{
-		return;
-	}
-
-	USkeletalMesh* Mesh = SkelMeshComponent->SkeletalMesh;
-	if (Mesh == nullptr)
-	{
-		return;
-	}
-
-	FSkeletalMeshLODRenderData& SkelMeshLODData = Mesh->GetResourceForRendering()->LODRenderData[LODIndex];
-	FSkinWeightVertexBuffer* SkinWeightBuffer = SkelMeshComponent->GetSkinWeightBuffer(LODIndex);
-	USkeletalMeshComponent::ComputeSkinnedPositions(SkelMeshComponent, TempPositions, InBoneMatrices, SkelMeshLODData, *SkinWeightBuffer);
-
-	if (bImportedVertices)
-	{
-		// Get the originally imported vertex numbers from the DCC.
-		const FSkeletalMeshModel* SkeletalMeshModel = Mesh->GetImportedModel();
-		const TArray<int32>& ImportedVertexNumbers = SkeletalMeshModel->LODModels[/*LODIndex=*/0].MeshToImportVertexMap;
-		if (ImportedVertexNumbers.Num() > 0)
-		{
-			const int32 MaxIndex = SkeletalMeshModel->LODModels[/*LODIndex=*/0].MaxImportVertex;
-
-			// Store the vertex positions for the original imported vertices (8 vertices for a cube).
-			OutPositions.AddZeroed(MaxIndex + 1);
-			for (int32 Index = 0; Index < TempPositions.Num(); ++Index)
-			{
-				const int32 ImportedVertex = ImportedVertexNumbers[Index];
-				OutPositions[ImportedVertex] = TempPositions[Index];
-			}
-		}
-	}
-	else
-	{
-		OutPositions = TempPositions;
-	}
-}
-
-void FMLDeformerEditorData::ExtractGeomCachePositions(int32 LODIndex, TArray<FVector3f>& TempPositions, TArray<FVector3f>& OutPositions) const
-{
-	OutPositions.Reset();
-	TempPositions.Reset();
-
-	UGeometryCacheComponent* GeomCacheComponent = GetEditorActor(EMLDeformerEditorActorIndex::Target).GeomCacheComponent;
-	UGeometryCache* GeomCache = GeomCacheComponent->GetGeometryCache();
-	if (GeomCache == nullptr)
-	{
-		return;
-	}
-
-	const float Time = GeomCacheComponent->GetAnimationTime();
-	TArray<FGeometryCacheMeshData> MeshData;
-	GeomCache->GetMeshDataAtTime(Time, MeshData);
-
-	if (MeshData.Num() > 0)
-	{
-		const TArray<uint32>& ImportedVertexNumbers = MeshData[0].ImportedVertexNumbers;
-		if (ImportedVertexNumbers.Num() > 0)
-		{
-			TempPositions = MeshData[0].Positions;
-
-			// Find the maximum value.
-			int32 MaxIndex = -1;
-			for (int32 Index = 0; Index < ImportedVertexNumbers.Num(); ++Index)
-			{
-				MaxIndex = FMath::Max(static_cast<int32>(ImportedVertexNumbers[Index]), MaxIndex);
-			}
-			check(MaxIndex > -1);
-
-			OutPositions.AddZeroed(MaxIndex + 1);
-
-			// Transform the positions with the alignment transform.
-			const FTransform& Transform = GetDeformerAsset()->AlignmentTransform;
-			for (int32 Index = 0; Index < TempPositions.Num(); ++Index)
-			{
-				const int32 ImportedVertex = static_cast<int32>(ImportedVertexNumbers[Index]);
-				const FVector Pos = TempPositions[Index];
-				OutPositions[ImportedVertex] = Transform.TransformPosition(Pos);
-			}
-		}
-	}
+	// Reinit the single frame cache.
+	FMLDeformerFrameCache::FInitSettings InitSettings;
+	InitSettings.CacheSizeInBytes = 0;	// This makes it just just use one frame, which is the minimum.
+	InitSettings.DeformerAsset = MLDeformerAsset.Get();
+	InitSettings.World = GetWorld();
+	InitSettings.bLogCacheStats = false;
+	SingleFrameCache.Init(InitSettings);
 }
 
 void FMLDeformerEditorData::SetAnimFrame(int32 FrameNumber)
@@ -437,33 +351,6 @@ void FMLDeformerEditorData::SetAnimFrame(int32 FrameNumber)
 	CurrentFrame = FrameNumber;
 }
 
-void FMLDeformerEditorData::UpdateDebugPointData()
-{
-	if (GetDeformerAsset()->GetVizSettings()->GetVisualizationMode() != EMLDeformerVizMode::TrainingData)
-	{
-		return;
-	}
-
-	const FMLDeformerEditorActor& BaseEditorActor = GetEditorActor(EMLDeformerEditorActorIndex::Base);
-	if (BaseEditorActor.SkelMeshComponent == nullptr || BaseEditorActor.SkelMeshComponent->SkeletalMesh == nullptr)
-	{
-		return;
-	}
-
-	// Extract the bone transforms and skinned vertex positions.
-	ExtractBoneTransforms(CurrentBoneMatrices, CurrentBoneTransforms);
-	ExtractSkinnedPositions(/*LODIndex=*/0, CurrentBoneMatrices, TempVectorBuffer, LinearSkinnedPositions);
-
-	const FMLDeformerEditorActor& TargetEditorActor = GetEditorActor(EMLDeformerEditorActorIndex::Target);
-	if (TargetEditorActor.GeomCacheComponent == nullptr || TargetEditorActor.GeomCacheComponent->GetGeometryCache() == nullptr)
-	{
-		return;
-	}
-
-	// Extract Geometry Cache positions.
-	ExtractGeomCachePositions(/*LODIndex=*/0, TempVectorBuffer, GeomCachePositions);
-}
-
 bool FMLDeformerEditorData::IsReadyForTraining() const
 {
 	UMLDeformerAsset* Asset = GetDeformerAsset();
@@ -493,99 +380,48 @@ bool FMLDeformerEditorData::IsReadyForTraining() const
 	return true;
 }
 
-bool FMLDeformerEditorData::GenerateTrainingData(uint32 LODIndex, uint32 FrameNumber, float DeltaCutoffLength, TArray<float>& OutDeltas, TArray<FTransform>& OutBoneTransforms, TArray<float>& OutCurveValues)
+bool FMLDeformerEditorData::GenerateDeltas(uint32 LODIndex, uint32 FrameNumber, TArray<float>& OutDeltas)
 {
 	OutDeltas.Reset();
-	OutBoneTransforms.Reset();
-	OutCurveValues.Reset();
 
 	if (!IsReadyForTraining())
 	{
 		return false;
 	}
 
-	// Calculate the number of frames to output.
-	UGeometryCacheComponent* GeometryCacheComponent = GetEditorActor(EMLDeformerEditorActorIndex::Target).GeomCacheComponent;
-	USkeletalMeshComponent* SkelMeshComponent = GetEditorActor(EMLDeformerEditorActorIndex::Base).SkelMeshComponent;
-	const uint32 AnimNumFrames = static_cast<uint32>(GeometryCacheComponent->GetNumberOfFrames());
-	const uint32 FrameIndex = FMath::Min(AnimNumFrames, FrameNumber);
-
-	// Calculate the vertex positions (and required transforms in order to perform skinning).
-	SetAnimFrame(static_cast<int32>(FrameIndex));
-	ExtractBoneTransforms(CurrentBoneMatrices, OutBoneTransforms);	// We just use our OutBoneTransforms array here, it will later be overwritten again.
-	ExtractSkinnedPositions(LODIndex, CurrentBoneMatrices, TempVectorBuffer, LinearSkinnedPositions);
-	ExtractGeomCachePositions(LODIndex, TempVectorBuffer, GeomCachePositions);
+	const FMLDeformerTrainingFrame& TrainingFrame = SingleFrameCache.GetTrainingFrameForAnimFrame(FrameNumber);
+	OutDeltas = TrainingFrame.GetVertexDeltas();
 
-	// Grab the actual required bone transforms
-	ExtractInputBoneTransforms(OutBoneTransforms);	// This internally resizes the output bone transforms array, so we overwrite existing contents from before.
-	ExtractInputCurveValues(OutCurveValues);
-
-	// Calculate the deltas.
-	if (LinearSkinnedPositions.Num() == GeomCachePositions.Num())
-	{
-		const int32 NumVertices = LinearSkinnedPositions.Num();
-		OutDeltas.AddUninitialized(3 * NumVertices);
-
-		for (int32 VertexIndex = 0; VertexIndex < NumVertices; ++VertexIndex)
-		{
-			const FVector Delta = GeomCachePositions[VertexIndex] - LinearSkinnedPositions[VertexIndex];
-			const int32 ArrayIndex = 3 * VertexIndex;
-			if (Delta.Length() < DeltaCutoffLength)
-			{
-				OutDeltas[ArrayIndex] = Delta.X;
-				OutDeltas[ArrayIndex + 1] = Delta.Y;
-				OutDeltas[ArrayIndex + 2] = Delta.Z;
-			}
-			else
-			{
-				OutDeltas[ArrayIndex] = 0.0f;
-				OutDeltas[ArrayIndex + 1] = 0.0f;
-				OutDeltas[ArrayIndex + 2] = 0.0f;
-			}
-		}
-	}
+	// Let's also directly extract the positions of the vertices, since we already calculated them.
+	const FMLDeformerSamplerData& SamplerData = SingleFrameCache.GetSampler().GetSamplerData();
+	LinearSkinnedPositions = SamplerData.GetSkinnedVertexPositions();
+	GeomCachePositions = SamplerData.GetGeometryCachePositions();
 
 	return true;
 }
 
-bool FMLDeformerEditorData::GenerateDeltas(uint32 LODIndex, uint32 FrameNumber, float DeltaCutoffLength, TArray<float>& OutDeltas)
-{
-	return GenerateTrainingData(LODIndex, FrameNumber, DeltaCutoffLength, OutDeltas, CurrentBoneTransforms, CurrentCurveWeights);
-}
-
-void FMLDeformerEditorData::UpdateVertexDeltaMeanAndScale(
-	const TArray<FVector3f>& InGeomCachePositions,
-	const TArray<FVector3f>& InLinearSkinnedPositions,
-	float DeltaCutoffLength,
-	FVector3f& InOutMeanVertexDelta,
-	FVector3f& InOutVertexDeltaScale,
-	float& InOutCount) const
+void FMLDeformerEditorData::UpdateVertexDeltaMeanAndScale(const FMLDeformerTrainingFrame& TrainingFrame, FVector3f& InOutMeanVertexDelta, FVector3f& InOutVertexDeltaScale, float& InOutCount)
 {
-	const int32 NumVertices = LinearSkinnedPositions.Num();
-	check(InGeomCachePositions.Num() == InLinearSkinnedPositions.Num());
+	check(!TrainingFrame.GetVertexDeltas().IsEmpty());
+	const int32 NumVertices = TrainingFrame.GetNumVertices();
 
 	FVector MeanDelta = FVector::ZeroVector;
 	FVector MinDelta(TNumericLimits<float>::Max());
 	FVector MaxDelta(-TNumericLimits<float>::Max());
-	float ValidVertexCount = 0.0f;
+	const TArray<float>& VertexDeltas = TrainingFrame.GetVertexDeltas();
 	for (int32 VertexIndex = 0; VertexIndex < NumVertices; ++VertexIndex)
 	{
-		const FVector Delta = GeomCachePositions[VertexIndex] - LinearSkinnedPositions[VertexIndex];
-		if (Delta.Length() < DeltaCutoffLength)
-		{
-			MeanDelta += Delta;
-			ValidVertexCount += 1.0f;
-
-			// Get minimum dimensions.
-			MinDelta = MinDelta.ComponentMin(Delta);
-			// Get maximum dimensions.
-			MaxDelta = MaxDelta.ComponentMax(Delta);
-		}
+		const int Offset = VertexIndex * 3;
+		const FVector Delta(VertexDeltas[Offset], VertexDeltas[Offset+1], VertexDeltas[Offset+2]);
+		MeanDelta += Delta;
+		MinDelta = MinDelta.ComponentMin(Delta);
+		MaxDelta = MaxDelta.ComponentMax(Delta);
 	}
 
-	if (ValidVertexCount > 0.0f)
+	const float ValidVertexCount = (float)NumVertices;
+	if (NumVertices > 0)
 	{
-		MeanDelta /= ValidVertexCount;
+		MeanDelta /= (float)NumVertices;
 	}
 
 	// Update global mean.
@@ -598,8 +434,9 @@ void FMLDeformerEditorData::UpdateVertexDeltaMeanAndScale(
 	InOutVertexDeltaScale = InOutVertexDeltaScale.ComponentMax(TempDiff.GetAbs());
 }
 
-bool FMLDeformerEditorData::ComputeVertexDeltaStatistics(uint32 LODIndex, float DeltaCutoffLength)
+bool FMLDeformerEditorData::ComputeVertexDeltaStatistics(uint32 LODIndex, FMLDeformerFrameCache* InFrameCache)
 {
+	check(InFrameCache);
 	if (!IsReadyForTraining() || bIsVertexDeltaNormalized)
 	{
 		return true;
@@ -610,38 +447,25 @@ bool FMLDeformerEditorData::ComputeVertexDeltaStatistics(uint32 LODIndex, float
 	const uint32 AnimNumFrames = static_cast<uint32>(GeometryCacheComponent->GetNumberOfFrames());
 
 	// Show some dialog with progress.
-	FSlowTask Task((float)AnimNumFrames, FText::FromString("Calculating data statistics"));
-	Task.Initialize();
+	const FText Title(LOCTEXT("PreprocessTrainingDataMessage", "Calculating data statistics"));
+	FScopedSlowTask Task((float)AnimNumFrames, Title);
 	Task.MakeDialog(true);
 
 	// Initialize mean vertex delta and vertex delta scale.
 	FVector3f VertexDeltaMean = FVector3f::ZeroVector;
 	FVector3f VertexDeltaScale = FVector3f::OneVector;
 	float Count = 0.0f;
-	TArray<FTransform> BoneTransforms;
+	USkeletalMeshComponent* SkeletalMeshComponent = GetEditorActor(EMLDeformerEditorActorIndex::Base).SkelMeshComponent;
 	for (uint32 FrameIndex = 0; FrameIndex < AnimNumFrames; FrameIndex++)
 	{
-		// Set the components to the right time values of the given frame.
-		SetAnimFrame(static_cast<int32>(FrameIndex));
-
-		// Extract the bone transforms.
-		ExtractBoneTransforms(CurrentBoneMatrices, BoneTransforms);
-
-		// Extract the linear skinned positions.
-		ExtractSkinnedPositions(LODIndex, CurrentBoneMatrices, TempVectorBuffer, LinearSkinnedPositions);
-
-		// Extract the geometry cache positions.
-		ExtractGeomCachePositions(LODIndex, TempVectorBuffer, GeomCachePositions);
-
-		// Update mean vertex delta.
-		UpdateVertexDeltaMeanAndScale(GeomCachePositions, LinearSkinnedPositions, DeltaCutoffLength,
-			VertexDeltaMean, VertexDeltaScale, Count);
+		// Get the deltas updated for this frame, and then update the mean and scale.
+		const FMLDeformerTrainingFrame& TrainingFrame = InFrameCache->GetTrainingFrameForAnimFrame(FrameIndex);
+		FMLDeformerEditorData::UpdateVertexDeltaMeanAndScale(TrainingFrame, VertexDeltaMean, VertexDeltaScale, Count);
 
 		// Forward the progress bar and check if we want to cancel.
 		Task.EnterProgressFrame();
 		if (Task.ShouldCancel())
 		{
-			Task.Destroy();
 			return false;
 		}
 	}
@@ -707,39 +531,6 @@ void FMLDeformerEditorData::SetHeatMapMaterialEnabled(bool bEnabled)
 	}
 }
 
-void FMLDeformerEditorData::ExtractInputCurveValues(TArray<float>& OutValues) const
-{
-	USkeletalMeshComponent* SkelMeshComponent = GetEditorActor(EMLDeformerEditorActorIndex::Base).SkelMeshComponent;
-	check(SkelMeshComponent);
-	const FMLDeformerInputInfo& InputInfo = GetDeformerAsset()->GetInputInfo();
-	InputInfo.ExtractCurveValues(SkelMeshComponent->GetAnimInstance(), OutValues);
-}
-
-void FMLDeformerEditorData::ExtractInputBoneTransforms(TArray<FTransform>& OutTransforms) const
-{
-	USkeletalMeshComponent* SkelMeshComponent = GetEditorActor(EMLDeformerEditorActorIndex::Base).SkelMeshComponent;
-	check(SkelMeshComponent);
-
-	// Get a copy of the bone transforms.
-	const TArray<FTransform> BoneTransforms = SkelMeshComponent->GetBoneSpaceTransforms();
-
-	const FMLDeformerInputInfo& InputInfo = GetDeformerAsset()->GetInputInfo();
-	const int32 NumBones = InputInfo.GetNumBones();
-	OutTransforms.Reset(NumBones);
-
-	if (NumBones == BoneTransforms.Num())
-	{
-		OutTransforms.AddUninitialized(NumBones);
-
-		for (int32 Index = 0; Index < NumBones; ++Index)
-		{
-			const FName BoneName = InputInfo.GetBoneName(Index);
-			const int32 SkelMeshBoneIndex = SkelMeshComponent->GetBoneIndex(BoneName);
-			OutTransforms[Index] = BoneTransforms[SkelMeshBoneIndex];
-		}
-	}
-}
-
 bool FMLDeformerEditorData::IsActorVisible(EMLDeformerEditorActorIndex ActorIndex) const
 {
 	const FMLDeformerEditorActor& EditorActor = GetEditorActor(ActorIndex);
@@ -941,4 +732,14 @@ void FMLDeformerEditorData::SetDefaultDeformerGraphIfNeeded()
 	}
 }
 
+UWorld* FMLDeformerEditorData::GetWorld() const
+{
+	return World;
+}
+
+void FMLDeformerEditorData::SetWorld(UWorld* InWorld)
+{
+	World = InWorld;
+}
+
 #undef LOCTEXT_NAMESPACE
diff --git a/Engine/Plugins/Experimental/MLDeformer/Source/MLDeformerEditor/Private/MLDeformerEditorMode.cpp b/Engine/Plugins/Experimental/MLDeformer/Source/MLDeformerEditor/Private/MLDeformerEditorMode.cpp
--- a/Engine/Plugins/Experimental/MLDeformer/Source/MLDeformerEditor/Private/MLDeformerEditorMode.cpp
+++ b/Engine/Plugins/Experimental/MLDeformer/Source/MLDeformerEditor/Private/MLDeformerEditorMode.cpp
@@ -246,19 +246,14 @@ void FMLDeformerEditorMode::Tick(FEditorViewportClient* ViewportClient, float De
 	UMLDeformerAsset* DeformerAsset = Data->GetDeformerAsset();
 	const UMLDeformerVizSettings* VizSettings = DeformerAsset->GetVizSettings();
 
-	// Calculate the vertex positions when needed.
-	if ((VizSettings->GetVisualizationMode() == EMLDeformerVizMode::TrainingData) && VizSettings->GetDrawVertex())
-	{
-		Data->UpdateDebugPointData();
-	}
-
 	UpdateActors();
 	UpdateLabels();
 
 	// Calculate the training deltas when needed.
-	if ((VizSettings->GetVisualizationMode() == EMLDeformerVizMode::TrainingData) && VizSettings->GetDrawVertexDeltas())
+	if ((VizSettings->GetVisualizationMode() == EMLDeformerVizMode::TrainingData) && 
+		(VizSettings->GetDrawVertex() || VizSettings->GetDrawVertexDeltas()))
 	{
-		Data->GenerateDeltas(0, VizSettings->FrameNumber, DeformerAsset->DeltaCutoffLength, Data->VertexDeltas);
+		Data->GenerateDeltas(0, VizSettings->FrameNumber, Data->VertexDeltas);
 	}
 
 	// Force us back into the pose we want to see.
diff --git a/Engine/Plugins/Experimental/MLDeformer/Source/MLDeformerEditor/Private/MLDeformerEditorToolkit.cpp b/Engine/Plugins/Experimental/MLDeformer/Source/MLDeformerEditor/Private/MLDeformerEditorToolkit.cpp
--- a/Engine/Plugins/Experimental/MLDeformer/Source/MLDeformerEditor/Private/MLDeformerEditorToolkit.cpp
+++ b/Engine/Plugins/Experimental/MLDeformer/Source/MLDeformerEditor/Private/MLDeformerEditorToolkit.cpp
@@ -172,6 +172,9 @@ void FMLDeformerEditorToolkit::FillToolbar(FToolBarBuilder& ToolbarBuilder)
 					{
 						ShowNotification(LOCTEXT("StartTraining", "Starting training process"), SNotificationItem::ECompletionState::CS_Pending, true);
 
+						// Change the interpolation type for the training sequence to step.
+						DeformerAsset->GetAnimSequence()->Interpolation = EAnimInterpolationType::Step;
+
 						// Initialize the training inputs.
 						DeformerAsset->SetInputInfo(DeformerAsset->CreateInputInfo());
 
@@ -179,8 +182,18 @@ void FMLDeformerEditorToolkit::FillToolbar(FToolBarBuilder& ToolbarBuilder)
 						// If this triggers, the train button most likely was enabled while it shouldn't be.
 						check(!DeformerAsset->GetInputInfo().IsEmpty());
 
+						// Init the frame cache.
+						const int32 NumFrames = (DeformerAsset->GetGeometryCache()->GetEndFrame() - DeformerAsset->GetGeometryCache()->GetStartFrame()) + 1;
+						FMLDeformerFrameCache::FInitSettings FrameCacheInitSettings;
+						FrameCacheInitSettings.DeformerAsset = DeformerAsset;
+						FrameCacheInitSettings.CacheSizeInBytes = 1024 * 1024 * DeformerAsset->GetCacheSizeInMegabytes();
+						FrameCacheInitSettings.World = EditorData->GetWorld();
+						TSharedPtr<FMLDeformerFrameCache> FrameCache = MakeShared<FMLDeformerFrameCache>();
+						FrameCache->Init(FrameCacheInitSettings);
+
 						// Perform training and load the resulting model.
 						MLDeformerModel->SetEditorData(EditorData);
+						MLDeformerModel->SetFrameCache(FrameCache);
 						MLDeformerModel->CreateDataSetInterface();
 
 						// Train ML deformer model using user-defined parameters.
@@ -192,13 +205,31 @@ void FMLDeformerEditorToolkit::FillToolbar(FToolBarBuilder& ToolbarBuilder)
 
 						UDebugSkelMeshComponent* SkelMeshComponent = EditorData->GetEditorActor(EMLDeformerEditorActorIndex::DeformedTest).SkelMeshComponent;
 						EditorData->InitAssets();
+						if (TrainingResult == ETrainingResult::Success || TrainingResult == ETrainingResult::Aborted)
+						{
+							// The InitAssets call above reset the normalized flag, so set it back to true.
+							// This is safe as we finished training, which means we already normalized data.
+							// If we aborted we still have normalized the data. Only when we have AbortedCantUse then we canceled the normalization process.
+							EditorData->bIsVertexDeltaNormalized = true;
+						}
 						EditorData->GetEditorActor(EMLDeformerEditorActorIndex::DeformedTest).MLDeformerComponent->SetupComponent(DeformerAsset, SkelMeshComponent);
 						if (bMarkDirty)
 						{
 							EditorData->GetDeformerAsset()->Modify();
 						}
+						
 						EditorData->GetDetailsView()->ForceRefresh();
 						EditorData->GetVizSettingsDetailsView()->ForceRefresh();
+
+						// Log memory usage.
+						const SIZE_T NumBytes = FrameCache->CalcMemUsageInBytes();;
+						UE_LOG(LogMLDeformer, Display, TEXT("Cache size used = %jd Bytes (%jd Kb or %.2f Mb)"), 
+							NumBytes, 
+							NumBytes / 1024, 
+							NumBytes / (float)(1024*1024));
+
+						// Clear the model internally, so it deletes the frame cache.
+						MLDeformerModel->Clear();
 					}
 					else
 					{
@@ -570,13 +601,16 @@ void FMLDeformerEditorToolkit::HandlePreviewSceneCreated(const TSharedRef<IPerso
 	// Load the default ML deformer graph asset.
 	EditorData->SetDefaultDeformerGraphIfNeeded();
 
+	// Set the world.
+	UWorld* World = InPersonaPreviewScene->GetWorld();
+	EditorData->SetWorld(World);
+
 	// Create the Linear Skinned (Base) actor.
 	const FLinearColor BaseLabelColor = FMLDeformerEditorStyle::Get().GetColor("MLDeformer.BaseMesh.LabelColor");
 	const FLinearColor BaseWireColor = FMLDeformerEditorStyle::Get().GetColor("MLDeformer.BaseMesh.WireframeColor");
 	CreateBaseActor(InPersonaPreviewScene, "Training Base", BaseLabelColor, BaseWireColor);
 
 	// Create the target actor (with Geometry Cache).
-	UWorld* World = InPersonaPreviewScene->GetWorld();
 	const FLinearColor TargetLabelColor = FMLDeformerEditorStyle::Get().GetColor("MLDeformer.TargetMesh.LabelColor");
 	const FLinearColor TargetWireColor = FMLDeformerEditorStyle::Get().GetColor("MLDeformer.TargetMesh.WireframeColor");
 	CreateGeomCacheActor(
@@ -629,7 +663,6 @@ void FMLDeformerEditorToolkit::HandlePreviewSceneCreated(const TSharedRef<IPerso
 
 	VizSettings->SetTempVisualizationMode(VizSettings->GetVisualizationMode());
 	EditorData->GetDeformerAsset()->SetTempTrainingInputs(EditorData->GetDeformerAsset()->GetTrainingInputs());
-	EditorData->UpdateDebugPointData();
 	OnSwitchedVisualizationMode();
 }
 
@@ -681,6 +714,12 @@ void FMLDeformerEditorToolkit::OnFinishedChangingDetails(const FPropertyChangedE
 		EditorData->GetVizSettingsDetailsView()->ForceRefresh();
 	}
 	else
+	if (Property->GetFName() == GET_MEMBER_NAME_CHECKED(UMLDeformerAsset, DeltaCutoffLength) ||
+	    Property->GetFName() == GET_MEMBER_NAME_CHECKED(UMLDeformerAsset, AlignmentTransform))
+	{
+		EditorData->InitAssets();
+	}
+	else
 	if (Property->GetFName() == GET_MEMBER_NAME_CHECKED(UMLDeformerAsset, TrainingInputs))
 	{
 		UMLDeformerAsset* DeformerAsset = EditorData->GetDeformerAsset();
diff --git a/Engine/Plugins/Experimental/MLDeformer/Source/MLDeformerEditor/Private/MLDeformerFrameCache.cpp b/Engine/Plugins/Experimental/MLDeformer/Source/MLDeformerEditor/Private/MLDeformerFrameCache.cpp
new file mode 100644
--- /dev/null
+++ b/Engine/Plugins/Experimental/MLDeformer/Source/MLDeformerEditor/Private/MLDeformerFrameCache.cpp
@@ -0,0 +1,593 @@
+// Copyright Epic Games, Inc. All Rights Reserved.
+
+#include "MLDeformerFrameCache.h"
+
+#include "Animation/DebugSkelMeshComponent.h"
+#include "Animation/AnimSequence.h"
+#include "Engine/SkeletalMesh.h"
+#include "Rendering/SkeletalMeshModel.h"
+#include "Rendering/SkeletalMeshLODModel.h"
+
+#include "GeometryCacheComponent.h"
+#include "GeometryCache.h"
+#include "GeometryCacheMeshData.h"
+
+#include "Math/NumericLimits.h"
+
+void FMLDeformerTrainingFrame::Clear()
+{
+	VertexDeltas.Empty();
+	BoneRotations.Empty();
+	CurveValues.Empty();
+	AnimFrameIndex = -1;
+}
+
+void FMLDeformerTrainingFrame::InitFromSamplerItem(int32 InAnimFrameIndex, const FMLDeformerSamplerData& InSamplerData)
+{
+	AnimFrameIndex = InAnimFrameIndex;
+	VertexDeltas = InSamplerData.GetVertexDeltas();
+	BoneRotations = InSamplerData.GetBoneRotations();
+	CurveValues = InSamplerData.GetCurveValues();
+}
+
+SIZE_T FMLDeformerTrainingFrame::CalcMemUsageInBytes() const
+{
+	SIZE_T NumBytes = 0;
+	NumBytes += VertexDeltas.GetAllocatedSize();
+	NumBytes += BoneRotations.GetAllocatedSize();
+	NumBytes += CurveValues.GetAllocatedSize();
+	return NumBytes;
+}
+
+//------------------------------------------------------------
+
+// Initializer a sampler item.
+void FMLDeformerSamplerData::Init(const FInitSettings& InitSettings)
+{
+	check(InitSettings.Sampler);
+	check(InitSettings.GeometryCacheComponent);
+	check(InitSettings.SkeletalMeshComponent);
+	check(InitSettings.SkeletalMeshComponent->SkeletalMesh);
+	check(InitSettings.GeometryCacheComponent->GetGeometryCache());
+	check(InitSettings.NumImportedVertices > 0);
+
+	NumImportedVertices = InitSettings.NumImportedVertices;
+	Sampler = InitSettings.Sampler;
+
+	SkeletalMeshComponent = InitSettings.SkeletalMeshComponent;
+	GeometryCacheComponent = InitSettings.GeometryCacheComponent;
+
+	// Resize arrays.
+	SkinnedVertexPositions.Empty();
+	SkinnedVertexPositions.AddUninitialized(NumImportedVertices);
+
+	GeomCacheVertexPositions.Empty();
+	GeomCacheVertexPositions.AddUninitialized(NumImportedVertices);
+
+	VertexDeltas.Empty(); // Multiply by 3 as we store this as an array of floats instead of an FVector3f.
+
+	BoneMatrices.Empty();
+	BoneRotations.Empty();
+	CurveValues.Empty();
+}
+
+void FMLDeformerSamplerData::Update(int32 InAnimFrameIndex)
+{
+	const int32 LODIndex = 0;
+
+	AnimFrameIndex = InAnimFrameIndex;
+
+	// Sample/tick the components at the exact frame time.
+	const UMLDeformerAsset& DeformerAsset = Sampler->GetDeformerAsset();
+	const float SampleTime = GeometryCacheComponent->GetTimeAtFrame(AnimFrameIndex);
+	if (SkeletalMeshComponent && SkeletalMeshComponent->SkeletalMesh)
+	{
+		// Sample the transforms at the frame time and extract the bone rotations.
+		SkeletalMeshComponent->SetPosition(SampleTime);
+		SkeletalMeshComponent->bPauseAnims = true;
+		SkeletalMeshComponent->RefreshBoneTransforms();
+		const FMLDeformerInputInfo& InputInfo = DeformerAsset.GetInputInfo();
+		InputInfo.ExtractBoneRotations(SkeletalMeshComponent, BoneRotations);
+
+		// Get the bone transform matrices used during skinning, then extract the skinned positions.
+		SkeletalMeshComponent->CacheRefToLocalMatrices(BoneMatrices);
+		ExtractSkinnedPositions(LODIndex, BoneMatrices, TempVertexPositions, SkinnedVertexPositions);
+
+		InputInfo.ExtractCurveValues(SkeletalMeshComponent, CurveValues);
+	}
+
+	// Sample the geometry cache vertices.
+	if (GeometryCacheComponent && GeometryCacheComponent->GetGeometryCache())
+	{
+		// Tick the component.
+		GeometryCacheComponent->SetManualTick(true);
+		GeometryCacheComponent->TickAtThisTime(SampleTime, false, false, false);
+
+		// Grab the vertex positions.
+		ExtractGeomCachePositions(LODIndex, TempVertexPositions, GeomCacheVertexPositions);
+	}
+
+	// Calculate the vertex deltas.
+	const float DeltaCutoffLength = DeformerAsset.GetDeltaCutoffLength();
+	CalculateVertexDeltas(SkinnedVertexPositions, GeomCacheVertexPositions, DeltaCutoffLength, VertexDeltas);
+}
+
+void FMLDeformerSamplerData::CalculateVertexDeltas(const TArray<FVector3f>& SkinnedPositions, const TArray<FVector3f>& GeomCachePositions, float DeltaCutoffLength, TArray<float>& OutVertexDeltas)
+{
+	check(SkinnedPositions.Num() == GeomCachePositions.Num());
+	const int32 NumVerts = SkinnedPositions.Num();
+	OutVertexDeltas.Reset(NumVerts * 3);
+	OutVertexDeltas.AddUninitialized(NumVerts * 3);
+	for (int32 VertexIndex = 0; VertexIndex < NumVerts; ++VertexIndex)
+	{
+		const FVector Delta = GeomCachePositions[VertexIndex] - SkinnedPositions[VertexIndex];
+		const int32 ArrayIndex = 3 * VertexIndex;
+		if (Delta.Length() < DeltaCutoffLength)
+		{
+			OutVertexDeltas[ArrayIndex] = Delta.X;
+			OutVertexDeltas[ArrayIndex + 1] = Delta.Y;
+			OutVertexDeltas[ArrayIndex + 2] = Delta.Z;
+		}
+		else
+		{
+			OutVertexDeltas[ArrayIndex] = 0.0f;
+			OutVertexDeltas[ArrayIndex + 1] = 0.0f;
+			OutVertexDeltas[ArrayIndex + 2] = 0.0f;
+		}
+	}
+}
+
+void FMLDeformerSamplerData::ExtractSkinnedPositions(int32 LODIndex, TArray<FMatrix44f>& InBoneMatrices, TArray<FVector3f>& TempPositions, TArray<FVector3f>& OutPositions) const
+{
+	OutPositions.Reset();
+	TempPositions.Reset();
+
+	if (SkeletalMeshComponent == nullptr)
+	{
+		return;
+	}
+
+	USkeletalMesh* Mesh = SkeletalMeshComponent->SkeletalMesh;
+	if (Mesh == nullptr)
+	{
+		return;
+	}
+
+	FSkeletalMeshLODRenderData& SkelMeshLODData = Mesh->GetResourceForRendering()->LODRenderData[LODIndex];
+	FSkinWeightVertexBuffer* SkinWeightBuffer = SkeletalMeshComponent->GetSkinWeightBuffer(LODIndex);
+	USkeletalMeshComponent::ComputeSkinnedPositions(SkeletalMeshComponent, TempPositions, InBoneMatrices, SkelMeshLODData, *SkinWeightBuffer);
+
+	// Get the originally imported vertex numbers from the DCC.
+	const FSkeletalMeshModel* SkeletalMeshModel = Mesh->GetImportedModel();
+	const TArray<int32>& ImportedVertexNumbers = SkeletalMeshModel->LODModels[LODIndex].MeshToImportVertexMap;
+	if (ImportedVertexNumbers.Num() > 0)
+	{
+		// Store the vertex positions for the original imported vertices (8 vertices for a cube).
+		OutPositions.AddZeroed(NumImportedVertices);
+		for (int32 Index = 0; Index < TempPositions.Num(); ++Index)
+		{
+			const int32 ImportedVertex = ImportedVertexNumbers[Index];
+			OutPositions[ImportedVertex] = TempPositions[Index];
+		}
+	}
+}
+
+void FMLDeformerSamplerData::ExtractGeomCachePositions(int32 LODIndex, TArray<FVector3f>& TempPositions, TArray<FVector3f>& OutPositions) const
+{
+	OutPositions.Reset();
+	TempPositions.Reset();
+
+	UGeometryCache* GeomCache = GeometryCacheComponent->GetGeometryCache();
+	if (GeomCache == nullptr)
+	{
+		return;
+	}
+
+	const float Time = GeometryCacheComponent->GetAnimationTime();
+	TArray<FGeometryCacheMeshData> MeshData;
+	GeomCache->GetMeshDataAtTime(Time, MeshData);
+
+	if (MeshData.Num() > 0)
+	{
+		const TArray<uint32>& ImportedVertexNumbers = MeshData[0].ImportedVertexNumbers;
+		if (ImportedVertexNumbers.Num() > 0)
+		{
+			TempPositions = MeshData[0].Positions;
+			const FTransform& AlignmentTransform = Sampler->GetDeformerAsset().GetAlignmentTransform();
+
+			// Transform the positions with the alignment transform.
+			OutPositions.AddZeroed(NumImportedVertices);
+			for (int32 Index = 0; Index < TempPositions.Num(); ++Index)
+			{
+				const int32 ImportedVertex = static_cast<int32>(ImportedVertexNumbers[Index]);
+				const FVector Pos = TempPositions[Index];
+				OutPositions[ImportedVertex] = AlignmentTransform.TransformPosition(Pos);
+			}
+		}
+	}
+}
+
+int32 FMLDeformerSamplerData::GetNumBones() const
+{
+	const UMLDeformerAsset& DeformerAsset = Sampler->GetDeformerAsset();
+	const FMLDeformerInputInfo& InputInfo = DeformerAsset.GetInputInfo();
+	return InputInfo.GetNumBones();
+}
+
+SIZE_T FMLDeformerSamplerData::CalcMemUsageInBytes() const
+{
+	SIZE_T NumBytes = 0;
+	NumBytes += SkinnedVertexPositions.GetAllocatedSize();
+	NumBytes += TempVertexPositions.GetAllocatedSize();
+	NumBytes += GeomCacheVertexPositions.GetAllocatedSize();
+	NumBytes += BoneMatrices.GetAllocatedSize();
+	NumBytes += VertexDeltas.GetAllocatedSize();
+	NumBytes += BoneRotations.GetAllocatedSize();
+	NumBytes += CurveValues.GetAllocatedSize();
+	return NumBytes;
+}
+
+//--------------------------------------------------------------------------------------------------------------
+
+FMLDeformerSampler::~FMLDeformerSampler()
+{
+	if (SkelMeshActor)
+	{
+		SkelMeshActor->Destroy();
+		SkelMeshActor = nullptr;
+	}
+
+	if (GeomCacheActor)
+	{
+		GeomCacheActor->Destroy();
+		GeomCacheActor = nullptr;
+	}
+}
+
+// Initialize the sampler.
+void FMLDeformerSampler::Init(const FInitSettings& InInitSettings)
+{
+	check(InInitSettings.DeformerAsset);
+	check(InInitSettings.DeformerAsset->SkeletalMesh);
+	check(InInitSettings.DeformerAsset->GeometryCache);
+	check(InInitSettings.DeformerAsset->AnimSequence);
+	check(InInitSettings.World);
+
+	// Copy the init settings.
+	InitSettings = InInitSettings;
+
+	// Create the skeletal mesh Actor.
+	if (SkelMeshActor == nullptr)
+	{
+		SkelMeshActor = CreateActor(InitSettings.World, "SkelMeshSamplerActor");
+		SkelMeshActor->SetActorTransform(FTransform::Identity);
+	}
+
+	// Create the geom cache actor.
+	if (GeomCacheActor == nullptr)
+	{
+		GeomCacheActor = CreateActor(InitSettings.World, "GeomCacheSamplerActor");
+		GeomCacheActor->SetActorTransform(InitSettings.DeformerAsset->GetAlignmentTransform());
+	}
+
+	// Extract the number vertices from the skeletal mesh and geometry cache, and make sure they match.
+	UMLDeformerAsset* DeformerAsset = InitSettings.DeformerAsset;
+	UGeometryCache* GeomCache = DeformerAsset->GetGeometryCache();
+	USkeletalMesh* SkeletalMesh = DeformerAsset->GetSkeletalMesh();
+	const int32 NumSkeletalMeshImportedVerts = ExtractNumImportedSkinnedVertices(DeformerAsset->GetSkeletalMesh());
+	const int32 NumGeomCacheImportedVerts = ExtractNumImportedGeomCacheVertices(GeomCache);
+	check(NumSkeletalMeshImportedVerts > 0);
+	check(NumGeomCacheImportedVerts > 0);
+	check(NumSkeletalMeshImportedVerts == NumGeomCacheImportedVerts);
+
+	// Create and init the items.
+	FMLDeformerSamplerData::FInitSettings DataInitSettings;
+	DataInitSettings.Sampler = this;
+	DataInitSettings.NumImportedVertices = NumSkeletalMeshImportedVerts;
+
+	// Create the skeletal mesh component.
+	UDebugSkelMeshComponent* SkelMeshComponent = SamplerData.GetSkeletalMeshComponent() ? SamplerData.GetSkeletalMeshComponent() : NewObject<UDebugSkelMeshComponent>(SkelMeshActor);
+	UAnimSequence* TrainingAnimSequence = DeformerAsset->GetAnimSequence();
+	SkelMeshComponent->SetSkeletalMesh(SkeletalMesh);
+	SkelMeshComponent->RegisterComponent();
+	DataInitSettings.SkeletalMeshComponent = SkelMeshComponent;
+	SkelMeshComponent->SetAnimationMode(EAnimationMode::AnimationSingleNode);
+	SkelMeshComponent->SetAnimation(TrainingAnimSequence);
+	SkelMeshComponent->SetPosition(0.0f);
+	SkelMeshComponent->SetPlayRate(1.0f);
+	SkelMeshComponent->Play(false);
+	SkelMeshComponent->SetVisibility(false);
+	SkelMeshComponent->RefreshBoneTransforms();
+
+	// Create the geometry cache component.
+	UGeometryCacheComponent* GeomCacheComponent = SamplerData.GetGeometryCacheComponent() ? SamplerData.GetGeometryCacheComponent() : NewObject<UGeometryCacheComponent>(GeomCacheActor);
+	GeomCacheComponent->SetGeometryCache(GeomCache);
+	GeomCacheComponent->SetManualTick(true);
+	GeomCacheComponent->RegisterComponent();
+	GeomCacheComponent->SetVisibility(false);
+	DataInitSettings.GeometryCacheComponent = GeomCacheComponent;
+
+	SamplerData.Init(DataInitSettings);
+
+	// Set the actor root components.
+	SkelMeshActor->SetRootComponent(SkelMeshComponent);
+	GeomCacheActor->SetRootComponent(GeomCacheComponent);
+}
+
+// Spawn an actor.
+AActor* FMLDeformerSampler::CreateActor(UWorld* InWorld, const FName& Name) const
+{
+	FActorSpawnParameters SpawnParams;
+	SpawnParams.Name = NAME_None;
+	AActor* Actor = InWorld->SpawnActor<AActor>(SpawnParams);
+	Actor->SetFlags(RF_Transient);
+	return Actor;
+}
+
+int32 FMLDeformerSampler::ExtractNumImportedSkinnedVertices(USkeletalMesh* SkeletalMesh) const
+{
+	check(SkeletalMesh);
+
+	// Extract the number of skinned mesh vertices.
+	int32 NumSkeletalMeshImportedVerts = 0;
+	const FSkeletalMeshModel* SkeletalMeshModel = SkeletalMesh->GetImportedModel();
+	if (SkeletalMeshModel)
+	{
+		const int32 MaxIndex = SkeletalMeshModel->LODModels[0].MaxImportVertex;
+		NumSkeletalMeshImportedVerts = (MaxIndex > 0) ? (MaxIndex + 1) : 0;
+	}
+	return NumSkeletalMeshImportedVerts;
+}
+
+int32 FMLDeformerSampler::ExtractNumImportedGeomCacheVertices(UGeometryCache* GeomCache) const
+{
+	check(GeomCache);
+	int32 NumGeomCacheImportedVerts = 0;
+
+	// Extract the geom cache number of imported vertices.
+	TArray<FGeometryCacheMeshData> MeshData;
+	GeomCache->GetMeshDataAtTime(0.0f, MeshData);
+	if (MeshData.Num() > 0)
+	{
+		const TArray<uint32>& ImportedVertexNumbers = MeshData[0].ImportedVertexNumbers;
+		if (ImportedVertexNumbers.Num() > 0)
+		{
+			// Find the maximum value.
+			int32 MaxIndex = -1;
+			for (int32 Index = 0; Index < ImportedVertexNumbers.Num(); ++Index)
+			{
+				MaxIndex = FMath::Max(static_cast<int32>(ImportedVertexNumbers[Index]), MaxIndex);
+			}
+			check(MaxIndex > -1);
+
+			NumGeomCacheImportedVerts = MaxIndex + 1;
+		}
+	}
+
+	return NumGeomCacheImportedVerts;
+}
+
+int32 FMLDeformerSampler::GetNumVertices() const
+{
+	return SamplerData.GetNumImportedVertices();
+}
+
+int32 FMLDeformerSampler::GetNumBones() const
+{
+	return SamplerData.GetNumBones();
+}
+
+int32 FMLDeformerSampler::GetNumCurves() const
+{
+	return SamplerData.GetCurveValues().Num();
+}
+
+int32 FMLDeformerSampler::GetNumFrames() const
+{
+	UMLDeformerAsset* DeformerAsset = InitSettings.DeformerAsset;
+	UGeometryCache* GeomCache = DeformerAsset->GetGeometryCache();
+	return (GeomCache->GetEndFrame() - GeomCache->GetStartFrame()) + 1;
+}
+
+SIZE_T FMLDeformerSampler::CalcMemUsageInBytes() const
+{
+	return SamplerData.CalcMemUsageInBytes();
+}
+
+//----------------------------------------------------------------------------------------------------------
+
+void FMLDeformerFrameCache::Init(const FInitSettings& InitSettings)
+{
+	check(InitSettings.CacheSizeInBytes >= 0);	// 0 is allowed, as we then simply have a 1 frame cache.
+	check(InitSettings.DeformerAsset);
+	check(InitSettings.World);
+
+	DeformerAsset = InitSettings.DeformerAsset;
+
+	if (DeformerAsset->GetSkeletalMesh() == nullptr || DeformerAsset->GetGeometryCache() == nullptr || DeformerAsset->GetAnimSequence() == nullptr)
+	{
+		FrameMap.Empty();
+		CachedTrainingFrames.Empty();
+		NextFreeCacheIndex = 0;
+		return;
+	}
+
+	// Initialize the sampler.
+	FMLDeformerSampler::FInitSettings SamplerInitSettings;
+	SamplerInitSettings.DeformerAsset = InitSettings.DeformerAsset;
+	SamplerInitSettings.World = InitSettings.World;
+	Sampler.Init(SamplerInitSettings);
+
+	// Create the first training frame, and update it already so we know how much memory one frame takes.
+	CachedTrainingFrames.Empty();
+	CachedTrainingFrames.AddDefaulted();
+	Sampler.GetSamplerData().Update(0);
+	CachedTrainingFrames[0].InitFromSamplerItem(0, Sampler.GetSamplerData());
+	const SIZE_T NumBytesPerTrainingFrame = CachedTrainingFrames[0].CalcMemUsageInBytes();
+
+	// Calculate how many frames fit inside our budget.
+	const int32 MaxCachedFrames = static_cast<int32>(InitSettings.CacheSizeInBytes / NumBytesPerTrainingFrame) + 1;
+	const int32 NumCachedFrames = FMath::Min(DeformerAsset->GetNumFrames(), MaxCachedFrames);	// Never use more than the number of frames in the training anim.
+
+	if (InitSettings.bLogCacheStats)
+	{
+		const SIZE_T EstimatedCacheSize = NumCachedFrames * NumBytesPerTrainingFrame;
+		UE_LOG(LogMLDeformer, Display, TEXT("ML Deformer cache will use %d frames in %jd bytes (%.2f Mb)"), NumCachedFrames, EstimatedCacheSize, EstimatedCacheSize / (float)(1024 * 1024));
+	}
+
+	// Allocate the cached frames ahead of time.
+	CachedTrainingFrames.AddDefaulted(NumCachedFrames - 1);	// -1 Because we already added one.
+
+	// Init the frame mapping table.
+	const int32 NumFrames = Sampler.GetNumFrames();
+	FrameMap.Empty();
+	FrameMap.AddUninitialized(NumFrames);
+	for (int32 FrameIndex = 0; FrameIndex < NumFrames; ++FrameIndex)
+	{
+		FrameMap[FrameIndex] = -1;	// -1 Means that this frame is not inside the list of cached frames.
+	}
+
+	NextFreeCacheIndex = 0;
+}
+
+int32 FMLDeformerFrameCache::GetNumVertices() const
+{
+	return Sampler.GetNumVertices();
+}
+
+int32 FMLDeformerFrameCache::GetNumBones() const
+{
+	return Sampler.GetNumBones();
+}
+
+int32 FMLDeformerFrameCache::GetNumCurves() const
+{
+	return Sampler.GetNumCurves();
+}
+
+int32 FMLDeformerFrameCache::GetCachedTrainingFrameIndex(int32 AnimFrameIndex) const
+{
+	return FrameMap[AnimFrameIndex];
+}
+
+void FMLDeformerFrameCache::Prefetch(int32 StartFrameIndex, int32 EndFrameIndex)
+{
+	check(IsValid());
+
+	int32 NumPrefetched = 0;
+	for (int32 AnimFrameIndex = StartFrameIndex; AnimFrameIndex <= EndFrameIndex; ++AnimFrameIndex)
+	{
+		// Generate the frame if we haven't got it cached already.
+		const int32 CachedFrameIndex = GetCachedTrainingFrameIndex(AnimFrameIndex);
+		if (CachedFrameIndex == -1)
+		{
+			GenerateFrame(AnimFrameIndex);
+		}
+
+		// We shouldn't prefetch more frames than the cache size.
+		if (++NumPrefetched >= CachedTrainingFrames.Num())
+		{
+			break;
+		}
+	}
+}
+
+const FMLDeformerTrainingFrame& FMLDeformerFrameCache::GetTrainingFrameForAnimFrame(int32 AnimFrameIndex)
+{
+	check(IsValid());
+
+	// Check if we already have this frame cached, and if not, generate it and insert it inside the cache.
+	int32 CachedFrameIndex = GetCachedTrainingFrameIndex(AnimFrameIndex);
+	if (CachedFrameIndex == -1)
+	{
+		CachedFrameIndex = GenerateFrame(AnimFrameIndex);	// This does heavy calculations and inserts the new frame in the cache, possibly replacing an older frame.
+	}
+	
+	return CachedTrainingFrames[CachedFrameIndex];
+}
+
+int32 FMLDeformerFrameCache::GetNextCacheFrameIndex()
+{
+	const int32 Result = NextFreeCacheIndex;
+
+	// Move the index to the next frame, wrap to the start of the cache buffer if needed.
+	NextFreeCacheIndex++;
+	if (NextFreeCacheIndex >= CachedTrainingFrames.Num())
+	{
+		// This prevents only caches the first calculated items, until the cache is full.
+		// It will keep overwriting the last item and keep the first ones in tact.
+		return CachedTrainingFrames.Num() - 1;	
+	}
+
+	return Result;
+}
+
+void FMLDeformerFrameCache::UpdateFrameMap()
+{
+	// First clear the values in the frame map.
+	// Basically set them to a value that tells that this frame hasn't been cached yet.
+	ResetFrameMap();
+
+	// For all cached frames, update the frame map to point to that cache item.
+	for (int32 Index = 0; Index < CachedTrainingFrames.Num(); ++Index)
+	{
+		const FMLDeformerTrainingFrame& TrainingFrame = CachedTrainingFrames[Index];
+		const int32 AnimFrameIndex = TrainingFrame.GetAnimFrameIndex();
+		if (AnimFrameIndex != -1)
+		{
+			FrameMap[AnimFrameIndex] = Index;
+		}
+	}
+}
+
+int32 FMLDeformerFrameCache::GenerateFrame(int32 AnimFrameIndex)
+{
+	check(AnimFrameIndex >= 0);
+
+	// Let the sampler item generate the deltas.
+	FMLDeformerSamplerData& SamplerData = Sampler.GetSamplerData();
+	SamplerData.Update(AnimFrameIndex);
+
+	// Extract the training data from the sampler item.
+	// Find a spot in the cache that we should modify, and update the frame map.
+	const int32 CachedFrameIndex = GetNextCacheFrameIndex();
+	FMLDeformerTrainingFrame& TrainingFrame = CachedTrainingFrames[CachedFrameIndex];
+	TrainingFrame.InitFromSamplerItem(AnimFrameIndex, SamplerData);
+	UpdateFrameMap();
+
+	return CachedFrameIndex;
+}
+
+SIZE_T FMLDeformerFrameCache::CalcMemUsageInBytes() const
+{
+	SIZE_T NumBytes = 0;
+	for (const FMLDeformerTrainingFrame& Frame : CachedTrainingFrames)
+	{
+		NumBytes += Frame.CalcMemUsageInBytes();
+	}
+	NumBytes += Sampler.CalcMemUsageInBytes();
+	return NumBytes;
+}
+
+void FMLDeformerFrameCache::ResetFrameMap()
+{
+	for (int32& Value : FrameMap)
+	{
+		Value = -1;
+	}
+}
+
+void FMLDeformerFrameCache::Clear()
+{
+	ResetFrameMap();
+	NextFreeCacheIndex = 0;
+	for (FMLDeformerTrainingFrame& Frame : CachedTrainingFrames)
+	{
+		Frame.Clear();
+	}
+}
+
+bool FMLDeformerFrameCache::IsValid() const
+{
+	return (!FrameMap.IsEmpty() && !CachedTrainingFrames.IsEmpty());
+}
diff --git a/Engine/Plugins/Experimental/MLDeformer/Source/MLDeformerEditor/Private/MLDeformerFrameCache.h b/Engine/Plugins/Experimental/MLDeformer/Source/MLDeformerEditor/Private/MLDeformerFrameCache.h
new file mode 100644
--- /dev/null
+++ b/Engine/Plugins/Experimental/MLDeformer/Source/MLDeformerEditor/Private/MLDeformerFrameCache.h
@@ -0,0 +1,199 @@
+// Copyright Epic Games, Inc. All Rights Reserved.
+
+#pragma once
+
+#include "MLDeformer.h"
+#include "MLDeformerAsset.h"
+
+#include "Math/Vector.h"
+#include "Math/Transform.h"
+#include "Containers/Array.h"
+#include "UObject/ObjectPtr.h"
+#include "UObject/NameTypes.h"
+
+class UDebugSkelMeshComponent;
+class UGeometryCacheComponent;
+class UAnimSequence;
+class USkeletalMesh;
+class UGeometryCache;
+class UWorld;
+class AActor;
+class FMLDeformerSampler;
+class FMLDeformerFrameCache;
+
+/**
+ * A sampler data object, which is basically a set of data that is used to generate training data for a given frame.
+ */
+class FMLDeformerSamplerData
+{
+public:
+	struct FInitSettings
+	{
+		FMLDeformerSampler* Sampler = nullptr;
+		UDebugSkelMeshComponent* SkeletalMeshComponent = nullptr;
+		UGeometryCacheComponent* GeometryCacheComponent = nullptr;
+		int32 NumImportedVertices = -1;
+	};
+
+	void Init(const FInitSettings& InitSettings);
+	void Update(int32 AnimFrameIndex);
+
+	const FMLDeformerSampler& GetSampler() const { return *Sampler; }
+
+	UDebugSkelMeshComponent* GetSkeletalMeshComponent() const { return SkeletalMeshComponent; }
+	UGeometryCacheComponent* GetGeometryCacheComponent() const { return GeometryCacheComponent; }
+
+	const TArray<FVector3f>& GetSkinnedVertexPositions() const { return SkinnedVertexPositions; }
+	const TArray<FVector3f>& GetGeometryCachePositions() const { return GeomCacheVertexPositions; }
+
+	const TArray<float>& GetVertexDeltas() const { return VertexDeltas; }
+	const TArray<float>& GetBoneRotations() const { return BoneRotations; }
+	const TArray<float>& GetCurveValues() const { return CurveValues; }
+
+	int32 GetNumImportedVertices() const { return NumImportedVertices; }
+	int32 GetNumBones() const;
+
+	SIZE_T CalcMemUsageInBytes() const;
+
+protected:
+	void ExtractSkinnedPositions(int32 LODIndex, TArray<FMatrix44f>& InBoneMatrices, TArray<FVector3f>& TempPositions, TArray<FVector3f>& OutPositions) const;
+	void ExtractGeomCachePositions(int32 LODIndex, TArray<FVector3f>& TempPositions, TArray<FVector3f>& OutPositions) const;
+	void CalculateVertexDeltas(const TArray<FVector3f>& SkinnedPositions, const TArray<FVector3f>& GeomCachePositions, float DeltaCutoffLength, TArray<float>& OutVertexDeltas);
+
+protected:
+	FMLDeformerSampler* Sampler = nullptr;
+	TObjectPtr<UDebugSkelMeshComponent>	SkeletalMeshComponent = nullptr;
+	TObjectPtr<UGeometryCacheComponent> GeometryCacheComponent = nullptr;
+	TArray<FVector3f> SkinnedVertexPositions;
+	TArray<FVector3f> TempVertexPositions;
+	TArray<FVector3f> GeomCacheVertexPositions;
+	TArray<FMatrix44f> BoneMatrices;
+	TArray<float> VertexDeltas;	// (NumImportedVerts * 3) -> xyz
+	TArray<float> BoneRotations; // (NumBones * 4) -> quat xyzw
+	TArray<float> CurveValues;
+	int32 NumImportedVertices = 0;
+	int32 AnimFrameIndex = -1;
+};
+
+/**
+ * The training data sampler, which is used to generate training data.
+ */
+class FMLDeformerSampler
+{
+public:
+	struct FInitSettings
+	{
+		UWorld* World = nullptr;
+		TObjectPtr<UMLDeformerAsset> DeformerAsset = nullptr;
+	};
+
+	FMLDeformerSampler() = default;
+	FMLDeformerSampler(const FMLDeformerSampler&) = delete;
+	FMLDeformerSampler(FMLDeformerSampler&&) = delete;
+	~FMLDeformerSampler();
+
+	FMLDeformerSampler& operator = (const FMLDeformerSampler&) = delete;
+	FMLDeformerSampler& operator = (FMLDeformerSampler&&) = delete;
+
+	void Init(const FInitSettings& InInitSetings);
+
+	FMLDeformerSamplerData& GetSamplerData() { return SamplerData; }
+	const FMLDeformerSamplerData& GetSamplerData() const { return SamplerData; }
+
+	int32 GetNumVertices() const;
+	int32 GetNumBones() const;
+	int32 GetNumCurves() const;
+	int32 GetNumFrames() const;
+
+	const FInitSettings& GetInitSettings() const { return InitSettings; }
+	const UMLDeformerAsset& GetDeformerAsset() const { return *InitSettings.DeformerAsset.Get(); }
+	UMLDeformerAsset& GetDeformerAsset() { return *InitSettings.DeformerAsset.Get(); }
+
+	SIZE_T CalcMemUsageInBytes() const;
+
+protected:
+	AActor* CreateActor(UWorld* InWorld, const FName& Name) const;
+	int32 ExtractNumImportedSkinnedVertices(USkeletalMesh* SkeletalMesh) const;
+	int32 ExtractNumImportedGeomCacheVertices(UGeometryCache* GeomCache) const;
+
+protected:
+	TObjectPtr<AActor> SkelMeshActor = nullptr;
+	TObjectPtr<AActor> GeomCacheActor = nullptr;
+	FMLDeformerSamplerData SamplerData;
+	FInitSettings InitSettings;
+};
+
+
+/**
+ * The training data for a given frame.
+ * Unlike the sampler frame, this contains only the data used during training and not all temp buffers used to generate this data.
+ * This data is already prepared to be passed directly to Python.
+ */
+class FMLDeformerTrainingFrame
+{
+public:
+	void Clear();
+	void InitFromSamplerItem(int32 InAnimFrameIndex, const FMLDeformerSamplerData& InSamplerData);
+
+	const TArray<float>& GetVertexDeltas() const { return VertexDeltas; }
+	const TArray<float>& GetBoneRotations() const { return BoneRotations; }
+	const TArray<float>& GetCurveValues() const { return CurveValues; }
+	int32 GetAnimFrameIndex() const { return AnimFrameIndex; }
+	int32 GetNumVertices() const { return VertexDeltas.Num() / 3; }
+
+	SIZE_T CalcMemUsageInBytes() const;
+
+protected:
+	TArray<float> VertexDeltas;	// NumVertices * 3, representing the vector x, y, z.
+	TArray<float> BoneRotations; // NumBones * 4, representing the quaternion x, y, z, w.
+	TArray<float> CurveValues;
+	int32 AnimFrameIndex = -1;
+};
+
+
+/**
+ * The training frame cache, which contains a subset of all training frames.
+ * Frames that are not inside the cache and are requested, will be generated on the fly.
+ * This is basically a FIFO cache of training data for given animation frames.
+ */
+class FMLDeformerFrameCache
+{
+public:
+	struct FInitSettings
+	{
+		SIZE_T CacheSizeInBytes = 1024 * 1024 * 1024 * 2ull;	// 2 gigabyte on default.
+		TObjectPtr<UMLDeformerAsset> DeformerAsset = nullptr;
+		UWorld* World = nullptr;
+		bool bLogCacheStats = true;
+	};
+
+	void Init(const FInitSettings& InitSettings);
+	void Prefetch(int32 StartFrameIndex, int32 EndFrameIndex);	// Use this to prefetch a given set of frames. The end frame specified is included as well.
+	void Clear(); // Clear the cached data, forcing everything to be regenerated.
+	const FMLDeformerTrainingFrame& GetTrainingFrameForAnimFrame(int32 AnimFrameIndex); // Get the training data for a given frame. This automatically generates on the fly it if needed.
+	bool IsValid() const;
+
+	int32 GetNumVertices() const;
+	int32 GetNumBones() const;
+	int32 GetNumCurves() const;
+
+	const UMLDeformerAsset& GetDeformerAsset() const { return *DeformerAsset.Get(); }
+	UMLDeformerAsset& GetDeformerAsset() { return *DeformerAsset.Get(); }
+	const FMLDeformerSampler& GetSampler() const { return Sampler; }
+
+	SIZE_T CalcMemUsageInBytes() const;
+
+protected:
+	int32 GenerateFrame(int32 AnimFrameIndex);	// Returns the new cached frame index
+	int32 GetCachedTrainingFrameIndex(int32 AnimFrameIndex) const;	// Returns -1 if not inside the cache.
+	int32 GetNextCacheFrameIndex(); // Gets the next cache spot to use when generating a new frame. This might point to already used cache spots. Implemented as FIFO.
+	void UpdateFrameMap();
+	void ResetFrameMap();
+
+protected:
+	FMLDeformerSampler Sampler;
+	TObjectPtr<UMLDeformerAsset> DeformerAsset = nullptr;
+	TArray<FMLDeformerTrainingFrame> CachedTrainingFrames;
+	TArray<int32> FrameMap;	// Map the global frame number to one in the CachedTrainingFrames array, or -1 when it's not inside the cache.
+	int32 NextFreeCacheIndex = 0;
+};
diff --git a/Engine/Plugins/Experimental/MLDeformer/Source/MLDeformerEditor/Private/MLDeformerPythonTrainingModel.cpp b/Engine/Plugins/Experimental/MLDeformer/Source/MLDeformerEditor/Private/MLDeformerPythonTrainingModel.cpp
--- a/Engine/Plugins/Experimental/MLDeformer/Source/MLDeformerEditor/Private/MLDeformerPythonTrainingModel.cpp
+++ b/Engine/Plugins/Experimental/MLDeformer/Source/MLDeformerEditor/Private/MLDeformerPythonTrainingModel.cpp
@@ -4,12 +4,29 @@
 #include "UObject/UObjectHash.h"
 #include "MLDeformerAsset.h"
 #include "MLDeformerEditorData.h"
+#include "MLDeformerFrameCache.h"
 #include "MLPytorchDataSetInterface.h"
 
 UMLDeformerPythonTrainingModel::UMLDeformerPythonTrainingModel()
 {
 }
 
+UMLDeformerPythonTrainingModel::~UMLDeformerPythonTrainingModel()
+{
+	Clear();
+}
+
+// Trigger the frame cache to be deleted.
+void UMLDeformerPythonTrainingModel::Clear()
+{
+	if (DataSetInterface)
+	{
+		DataSetInterface->Clear();
+	}
+	EditorData.Reset();
+	FrameCache.Reset();
+}
+
 UMLDeformerPythonTrainingModel* UMLDeformerPythonTrainingModel::Get()
 {
 	TArray<UClass*> PythonTrainingModels;
@@ -34,9 +51,15 @@ void UMLDeformerPythonTrainingModel::CreateDataSetInterface()
 		DataSetInterface = NewObject<UMLPytorchDataSetInterface>();
 	}
 	DataSetInterface->SetEditorData(EditorData);
+	DataSetInterface->SetFrameCache(FrameCache);
 }
 
 void UMLDeformerPythonTrainingModel::SetEditorData(TSharedPtr<FMLDeformerEditorData> InEditorData)
 {
 	EditorData = InEditorData;
 }
+
+void UMLDeformerPythonTrainingModel::SetFrameCache(TSharedPtr<FMLDeformerFrameCache> InFrameCache)
+{
+	FrameCache = InFrameCache;
+}
diff --git a/Engine/Plugins/Experimental/MLDeformer/Source/MLDeformerEditor/Private/MLDeformerPythonTrainingModel.h b/Engine/Plugins/Experimental/MLDeformer/Source/MLDeformerEditor/Private/MLDeformerPythonTrainingModel.h
--- a/Engine/Plugins/Experimental/MLDeformer/Source/MLDeformerEditor/Private/MLDeformerPythonTrainingModel.h
+++ b/Engine/Plugins/Experimental/MLDeformer/Source/MLDeformerEditor/Private/MLDeformerPythonTrainingModel.h
@@ -9,6 +9,7 @@
 class UMLPytorchDataSetInterface;
 class FMLDeformerEditorData;
 class UMLDeformerAsset;
+class FMLDeformerFrameCache;
 
 /** Training process return codes. */
 UENUM()
@@ -30,6 +31,9 @@ class MLDEFORMEREDITOR_API UMLDeformerPythonTrainingModel : public UObject
 	GENERATED_BODY()
 public:
 	UMLDeformerPythonTrainingModel();
+	~UMLDeformerPythonTrainingModel();
+
+	void Clear();
 
 	UFUNCTION(BlueprintCallable, Category = Python)
 	static UMLDeformerPythonTrainingModel* Get();
@@ -56,7 +60,9 @@ public:
 	UMLPytorchDataSetInterface* DataSetInterface = nullptr;
 
 	void SetEditorData(TSharedPtr<FMLDeformerEditorData> InEditorData);
+	void SetFrameCache(TSharedPtr<FMLDeformerFrameCache> InFrameCache);
 
 private:
 	TSharedPtr<FMLDeformerEditorData> EditorData = nullptr;
+	TSharedPtr<FMLDeformerFrameCache> FrameCache = nullptr;
 };
diff --git a/Engine/Plugins/Experimental/MLDeformer/Source/MLDeformerEditor/Private/MLPytorchDataSetInterface.cpp b/Engine/Plugins/Experimental/MLDeformer/Source/MLDeformerEditor/Private/MLPytorchDataSetInterface.cpp
--- a/Engine/Plugins/Experimental/MLDeformer/Source/MLDeformerEditor/Private/MLPytorchDataSetInterface.cpp
+++ b/Engine/Plugins/Experimental/MLDeformer/Source/MLDeformerEditor/Private/MLPytorchDataSetInterface.cpp
@@ -6,6 +6,7 @@
 #include "MLDeformer.h"
 #include "MLDeformerEditorData.h"
 #include "MLDeformerAsset.h"
+#include "MLDeformerFrameCache.h"
 
 #include "Animation/DebugSkelMeshComponent.h"
 #include "GeometryCache.h"
@@ -14,11 +15,27 @@ UMLPytorchDataSetInterface::UMLPytorchDataSetInterface()
 {
 }
 
+UMLPytorchDataSetInterface::~UMLPytorchDataSetInterface()
+{
+	Clear();
+}
+
+void UMLPytorchDataSetInterface::Clear()
+{
+	EditorData.Reset();
+	FrameCache.Reset();
+}
+
 void UMLPytorchDataSetInterface::SetEditorData(TSharedPtr<FMLDeformerEditorData> InEditorData)
 {
 	EditorData = InEditorData;
 }
 
+void UMLPytorchDataSetInterface::SetFrameCache(TSharedPtr<FMLDeformerFrameCache> InFrameCache)
+{
+	FrameCache = InFrameCache;
+}
+
 bool UMLPytorchDataSetInterface::IsValid() const
 {
 	const FMLDeformerEditorData* Data = EditorData.Get();
@@ -66,12 +83,11 @@ bool UMLPytorchDataSetInterface::SetCurrentSampleIndex(int32 Index)
 		return false;
 	}
 
-	// Generate and store the deltas.
-	const float DeltaCutoffLength = DeformerAsset->GetDeltaCutoffLength();
-	if (!Data->GenerateTrainingData(0, Index, DeltaCutoffLength, SampleDeltas, SampleBoneTransforms, SampleCurveWeights))
-	{
-		return false;
-	}
+	// Generate and store the training data.
+	const FMLDeformerTrainingFrame& TrainingFrame = FrameCache->GetTrainingFrameForAnimFrame(Index);
+	SampleDeltas = TrainingFrame.GetVertexDeltas();
+	SampleBoneRotations = TrainingFrame.GetBoneRotations();
+	SampleCurveValues = TrainingFrame.GetCurveValues();
 
 	return true;
 }
@@ -85,7 +101,7 @@ bool UMLPytorchDataSetInterface::ComputeDeltasStatistics()
 
 	// Generate and store the deltas.
 	const float DeltaCutoffLength = DeformerAsset->GetDeltaCutoffLength();
-	if (Data->ComputeVertexDeltaStatistics(0, DeltaCutoffLength))
+	if (Data->ComputeVertexDeltaStatistics(0, FrameCache.Get()))
 	{
 		//Update mean vertex delta and vertex data scale.
 		VertexDeltaMean.X = DeformerAsset->VertexDeltaMean.X;
diff --git a/Engine/Plugins/Experimental/MLDeformer/Source/MLDeformerEditor/Public/MLDeformerEditorData.h b/Engine/Plugins/Experimental/MLDeformer/Source/MLDeformerEditor/Public/MLDeformerEditorData.h
--- a/Engine/Plugins/Experimental/MLDeformer/Source/MLDeformerEditor/Public/MLDeformerEditorData.h
+++ b/Engine/Plugins/Experimental/MLDeformer/Source/MLDeformerEditor/Public/MLDeformerEditorData.h
@@ -5,6 +5,7 @@
 #include "CoreMinimal.h"
 #include "MLDeformerAsset.h"
 #include "MLDeformerInputInfo.h"
+#include "MLDeformerFrameCache.h"
 
 class UAnimInstance;
 class UMLDeformerAsset;
@@ -23,6 +24,8 @@ class UMaterial;
 class AActor;
 class UTextRenderComponent;
 class SSimpleTimeSlider;
+class UWorld;
+class FMLDeformerFrameCache;
 struct FMLDeformerTrainingData;
 struct FReferenceSkeleton;
 struct FGeometryCacheMeshData;
@@ -59,6 +62,7 @@ public:
 	void SetEditorToolkit(FMLDeformerEditorToolkit* InToolkit);
 	void SetVizSettingsDetailsView(TSharedPtr<IDetailsView> InDetailsView);
 	void SetTimeSlider(TSharedPtr<SSimpleTimeSlider> InTimeSlider);
+	void SetWorld(UWorld* InWorld);
 
 	int32 GetNumEditorActors() const;
 	void SetEditorActor(EMLDeformerEditorActorIndex Index, const FMLDeformerEditorActor& Actor);
@@ -68,9 +72,8 @@ public:
 	bool IsTrainingActor(EMLDeformerEditorActorIndex Index) const;
 
 	void InitAssets();
-	bool GenerateDeltas(uint32 LODIndex, uint32 FrameNumber, float DeltaCutoffLength, TArray<float>& OutDeltas);
-	bool GenerateTrainingData(uint32 LODIndex, uint32 FrameNumber, float DeltaCutoffLength, TArray<float>& OutDeltas, TArray<FTransform>& OutBoneTransforms, TArray<float>& OutCurveValues);
-	bool ComputeVertexDeltaStatistics(uint32 LODIndex, float DeltaCutoffLength);
+	bool GenerateDeltas(uint32 LODIndex, uint32 FrameNumber, TArray<float>& OutDeltas);
+	bool ComputeVertexDeltaStatistics(uint32 LODIndex, FMLDeformerFrameCache* FrameCache);
 
 	void SetTimeSliderRange(double StartTime, double EndTime);
 	void UpdateTimeSlider();
@@ -88,6 +91,7 @@ public:
 	IDetailsView* GetVizSettingsDetailsView() const;
 	FMLDeformerEditorToolkit* GetEditorToolkit() const;
 	SSimpleTimeSlider* GetTimeSlider() const;
+	UWorld* GetWorld() const;
 
 	float GetDuration() const;
 	float GetSnappedFrameTime(float InTime) const;
@@ -97,24 +101,13 @@ public:
 	/**
 	 *Update mean and scale of vertex deltas.
 	 *
-	 * @param InGeomCachePositions Geometric cache positions
-	 * @param InLinearSkinnedPositions Skinned positions
-	 * @param InDeltaCutoffLength Delta cutoff length
+	 * @param TrainingFrame The training frame that contains the vertex deltas.
 	 * @param InOutMeanVertexDelta Mean vertex delta
 	 * @param InOutVertexDeltaScale Vertex delta scale
 	 * @param InOutCount Count
 	 */
-	void UpdateVertexDeltaMeanAndScale(const TArray<FVector3f>& InGeomCachePositions, const TArray<FVector3f>& InLinearSkinnedPositions,
-		float InDeltaCutoffLength, FVector3f& InOutMeanVertexDelta, FVector3f& InOutVertexDeltaScale, float& InOutCount) const;
+	static void UpdateVertexDeltaMeanAndScale(const FMLDeformerTrainingFrame& TrainingFrame, FVector3f& InOutMeanVertexDelta, FVector3f& InOutVertexDeltaScale, float& InOutCount);
 
-	void ExtractBoneTransforms(TArray<FMatrix44f>& OutBoneMatrices, TArray<FTransform>& OutBoneTransforms) const;
-	void ExtractSkinnedPositions(int32 LODIndex, TArray<FMatrix44f>& InBoneMatrices, TArray<FVector3f>& TempPositions, TArray<FVector3f>& OutPositions, bool bImportedVertices=true) const;
-	void ExtractGeomCachePositions(int32 LODIndex, TArray<FVector3f>& TempPositions, TArray<FVector3f>& OutPositions) const;
-
-	void ExtractInputBoneTransforms(TArray<FTransform>& OutTransforms) const;
-	void ExtractInputCurveValues(TArray<float>& OutValues) const;
-
-	void UpdateDebugPointData();
 	void UpdateTestAnimPlaySpeed();
 	void UpdateDeformerGraph();
 	static FString GetDefaultDeformerGraphAssetPath();
@@ -130,17 +123,14 @@ public:
 	void SetDefaultDeformerGraphIfNeeded();
 	bool IsActorVisible(EMLDeformerEditorActorIndex ActorIndex) const;
 
+	FMLDeformerFrameCache& GetSingleFrameCache() { return SingleFrameCache; }
+	const FMLDeformerFrameCache& GetSingleFrameCache() const { return SingleFrameCache; }
+	
 public:
-	TArray<FMatrix44f> CurrentBoneMatrices;
-	TArray<FTransform> CurrentBoneTransforms;
-	TArray<float> CurrentCurveWeights;
-
 	TArray<FVector3f> LinearSkinnedPositions;
 	TArray<FVector3f> GeomCachePositions;
-	TArray<FVector3f> TempVectorBuffer;
 	TArray<float> VertexDeltas;
 	bool bIsVertexDeltaNormalized = false;
-	TArray<FGeometryCacheMeshData> MeshDatas;
 
 	int32 CurrentFrame = -1;
 
@@ -169,6 +159,12 @@ private:
 	/** The heatmap material. */
 	TObjectPtr<UMaterial> HeatMapMaterial;
 
+	/** The world that our actors are inside. */
+	TObjectPtr<UWorld> World;
+
 	/** The timeline slider widget. */
 	TSharedPtr<SSimpleTimeSlider> TimeSlider;
+
+	/** Single frame cache, used to calculate the training data for the current frame in the timeline. */
+	FMLDeformerFrameCache SingleFrameCache;
 };
diff --git a/Engine/Plugins/Experimental/MLDeformer/Source/MLDeformerEditor/Public/MLPytorchDataSetInterface.h b/Engine/Plugins/Experimental/MLDeformer/Source/MLDeformerEditor/Public/MLPytorchDataSetInterface.h
--- a/Engine/Plugins/Experimental/MLDeformer/Source/MLDeformerEditor/Public/MLPytorchDataSetInterface.h
+++ b/Engine/Plugins/Experimental/MLDeformer/Source/MLDeformerEditor/Public/MLPytorchDataSetInterface.h
@@ -7,6 +7,8 @@
 #include "MLPytorchDataSetInterface.generated.h"
 
 class FMLDeformerEditorData;
+class FMLDeformerFrameCache;
+
 /**
  * 
  */
@@ -16,7 +18,11 @@ class MLDEFORMEREDITOR_API UMLPytorchDataSetInterface : public UObject
 	GENERATED_BODY()
 public:
 	UMLPytorchDataSetInterface(); 
+	~UMLPytorchDataSetInterface();
+
+	void Clear();
 	void SetEditorData(TSharedPtr<FMLDeformerEditorData> InEditorData);
+	void SetFrameCache(TSharedPtr<FMLDeformerFrameCache> InFrameCache);
 
 	/** Get the number of input transforms. */
 	UFUNCTION(BlueprintPure, Category = "Training Data")
@@ -42,13 +48,17 @@ public:
 	UFUNCTION(BlueprintCallable, Category = "Training Data")
 	bool ComputeDeltasStatistics();
 
-	/** The delta values per vertex for this sample. This is updated after SetCurrentSampleIndex is called. */
+	/** The delta values per vertex for this sample. This is updated after SetCurrentSampleIndex is called. Contains an xyz (3 floats) for each vertex. */
 	UPROPERTY(EditAnywhere, BlueprintReadWrite, Category = "Training Data")
 	TArray<float> SampleDeltas;
 
 	/** The curve weights. This is updated after SetCurrentSampleIndex is called. */
 	UPROPERTY(EditAnywhere, BlueprintReadWrite, Category = "Training Data")
-	TArray<float> SampleCurveWeights;
+	TArray<float> SampleCurveValues;
+
+	/** The bone rotations in bone (local) space for this sample. This is updated after SetCurrentSampleIndex is called. Contains an xyzw (4 floats) for each bone. */
+	UPROPERTY(EditAnywhere, BlueprintReadWrite, Category = "Training Data")
+	TArray<float> SampleBoneRotations;
 
 	/** Mean delta computed over the entire dataset. */
 	UPROPERTY(EditAnywhere, BlueprintReadWrite, Category = "Training Data")
@@ -58,10 +68,6 @@ public:
 	UPROPERTY(EditAnywhere, BlueprintReadWrite, Category = "Training Data")
 	FVector VertexDeltaScale = FVector::OneVector;
 
-	/** The bone transformations in component space for this sample. This is updated after SetCurrentSampleIndex is called. */
-	UPROPERTY(EditAnywhere, BlueprintReadWrite, Category = "Training Data")
-	TArray<FTransform> SampleBoneTransforms;
-
 	/** The current sample index. */
 	UPROPERTY(EditAnywhere, BlueprintReadOnly, Category = "Training Data")
 	int32 CurrentSampleIndex = -1;
@@ -71,4 +77,5 @@ private:
 
 private:
 	TSharedPtr<FMLDeformerEditorData> EditorData;
+	TSharedPtr<FMLDeformerFrameCache> FrameCache;
 };
