[MLDeformer] Fix a crash when there is no AnimInstance on a SkeletalMesh, when dealing with curves.

#trivial

diff --git a/Engine/Plugins/Experimental/MLDeformer/Source/MLDeformer/Private/MLDeformerInstance.cpp b/Engine/Plugins/Experimental/MLDeformer/Source/MLDeformer/Private/MLDeformerInstance.cpp
--- a/Engine/Plugins/Experimental/MLDeformer/Source/MLDeformer/Private/MLDeformerInstance.cpp
+++ b/Engine/Plugins/Experimental/MLDeformer/Source/MLDeformer/Private/MLDeformerInstance.cpp
@@ -110,11 +110,21 @@ int64 FMLDeformerInstance::SetCurveValues(float* OutputBuffer, int64 OutputBuffe
 
 	// Write the curve weights to the output buffer.
 	UAnimInstance* AnimInstance = SkeletalMeshComponent->GetAnimInstance();
-	for (int32 CurveIndex = 0; CurveIndex < AssetNumCurves; ++CurveIndex)
+	if (AnimInstance)
 	{
-		const FName CurveName = InputInfo.GetCurveName(CurveIndex);
-		const float CurveValue = AnimInstance->GetCurveValue(CurveName);	// Outputs 0.0 when not found.
-		OutputBuffer[Index++] = CurveValue;
+		for (int32 CurveIndex = 0; CurveIndex < AssetNumCurves; ++CurveIndex)
+		{
+			const FName CurveName = InputInfo.GetCurveName(CurveIndex);
+			const float CurveValue = AnimInstance->GetCurveValue(CurveName);	// Outputs 0.0 when not found.
+			OutputBuffer[Index++] = CurveValue;
+		}
+	}
+	else
+	{
+		for (int32 CurveIndex = 0; CurveIndex < AssetNumCurves; ++CurveIndex)
+		{
+			OutputBuffer[Index++] = 0.0f;
+		}
 	}
 
 	return Index;
@@ -133,7 +143,11 @@ void FMLDeformerInstance::SetNeuralNetworkInputValues(float* InputData, int64 Nu
 
 void FMLDeformerInstance::Update()
 {
-	if (DeformerAsset == nullptr || SkeletalMeshComponent == nullptr || SkeletalMeshComponent->SkeletalMesh == nullptr || !bIsCompatible)
+	// Some safety checks.
+	if (DeformerAsset == nullptr || 
+		SkeletalMeshComponent == nullptr || 
+		SkeletalMeshComponent->SkeletalMesh == nullptr || 
+		!bIsCompatible)
 	{
 		return;
 	}
