[MLDeformer]
- Fixed a crash bug when having an out of range frame.
- Frames are now clamped to the maximum range in the UI.

diff --git a/Engine/Plugins/Experimental/MLDeformer/Source/MLDeformerEditor/Private/MLDeformerEditorData.cpp b/Engine/Plugins/Experimental/MLDeformer/Source/MLDeformerEditor/Private/MLDeformerEditorData.cpp
--- a/Engine/Plugins/Experimental/MLDeformer/Source/MLDeformerEditor/Private/MLDeformerEditorData.cpp
+++ b/Engine/Plugins/Experimental/MLDeformer/Source/MLDeformerEditor/Private/MLDeformerEditorData.cpp
@@ -314,6 +314,8 @@ void FMLDeformerEditorData::InitAssets()
 	MLDeformerAsset->UpdateCachedNumVertices();
 	UpdateDeformerGraph();
 
+	ClampFrameIndex();
+
 	// Reinit the single frame cache.
 	FMLDeformerFrameCache::FInitSettings InitSettings;
 	InitSettings.CacheSizeInBytes = 0;	// This makes it just just use one frame, which is the minimum.
@@ -330,6 +332,8 @@ void FMLDeformerEditorData::SetAnimFrame(int32 FrameNumber)
 		return;
 	}
 
+	ClampFrameIndex();
+
 	// Calculate the time in the animation data to visualize.
 	const float TimeOffset = GetTimeAtFrame(FrameNumber);
 
@@ -742,4 +746,17 @@ void FMLDeformerEditorData::SetWorld(UWorld* InWorld)
 	World = InWorld;
 }
 
+void FMLDeformerEditorData::ClampFrameIndex()
+{
+	UMLDeformerVizSettings* VizSettings = GetDeformerAsset()->GetVizSettings();
+	if (MLDeformerAsset->GetNumFrames() > 0)
+	{
+		VizSettings->FrameNumber = FMath::Min(VizSettings->FrameNumber, static_cast<uint32>(GetDeformerAsset()->GetNumFrames() - 1));
+	}
+	else
+	{
+		VizSettings->FrameNumber = 0;
+	}
+}
+
 #undef LOCTEXT_NAMESPACE
diff --git a/Engine/Plugins/Experimental/MLDeformer/Source/MLDeformerEditor/Private/MLDeformerEditorMode.cpp b/Engine/Plugins/Experimental/MLDeformer/Source/MLDeformerEditor/Private/MLDeformerEditorMode.cpp
--- a/Engine/Plugins/Experimental/MLDeformer/Source/MLDeformerEditor/Private/MLDeformerEditorMode.cpp
+++ b/Engine/Plugins/Experimental/MLDeformer/Source/MLDeformerEditor/Private/MLDeformerEditorMode.cpp
@@ -243,13 +243,14 @@ void FMLDeformerEditorMode::Tick(FEditorViewportClient* ViewportClient, float De
 {
 	FMLDeformerEditorData* Data = EditorData.Pin().Get();
 
-	UMLDeformerAsset* DeformerAsset = Data->GetDeformerAsset();
-	const UMLDeformerVizSettings* VizSettings = DeformerAsset->GetVizSettings();
+	Data->ClampFrameIndex();
 
 	UpdateActors();
 	UpdateLabels();
 
 	// Calculate the training deltas when needed.
+	UMLDeformerAsset* DeformerAsset = Data->GetDeformerAsset();
+	const UMLDeformerVizSettings* VizSettings = DeformerAsset->GetVizSettings();
 	if ((VizSettings->GetVisualizationMode() == EMLDeformerVizMode::TrainingData) && 
 		(VizSettings->GetDrawVertex() || VizSettings->GetDrawVertexDeltas()))
 	{
diff --git a/Engine/Plugins/Experimental/MLDeformer/Source/MLDeformerEditor/Private/MLDeformerFrameCache.cpp b/Engine/Plugins/Experimental/MLDeformer/Source/MLDeformerEditor/Private/MLDeformerFrameCache.cpp
--- a/Engine/Plugins/Experimental/MLDeformer/Source/MLDeformerEditor/Private/MLDeformerFrameCache.cpp
+++ b/Engine/Plugins/Experimental/MLDeformer/Source/MLDeformerEditor/Private/MLDeformerFrameCache.cpp
@@ -467,7 +467,7 @@ int32 FMLDeformerFrameCache::GetNumCurves() const
 
 int32 FMLDeformerFrameCache::GetCachedTrainingFrameIndex(int32 AnimFrameIndex) const
 {
-	return FrameMap[AnimFrameIndex];
+	return FrameMap.IsValidIndex(AnimFrameIndex) ? FrameMap[AnimFrameIndex] : -1;
 }
 
 void FMLDeformerFrameCache::Prefetch(int32 StartFrameIndex, int32 EndFrameIndex)
diff --git a/Engine/Plugins/Experimental/MLDeformer/Source/MLDeformerEditor/Public/MLDeformerEditorData.h b/Engine/Plugins/Experimental/MLDeformer/Source/MLDeformerEditor/Public/MLDeformerEditorData.h
--- a/Engine/Plugins/Experimental/MLDeformer/Source/MLDeformerEditor/Public/MLDeformerEditorData.h
+++ b/Engine/Plugins/Experimental/MLDeformer/Source/MLDeformerEditor/Public/MLDeformerEditorData.h
@@ -114,6 +114,7 @@ public:
 	static UComputeGraph* LoadDefaultDeformerGraph();
 	void SetAnimFrame(int32 FrameNumber);
 	void CreateMaterials();
+	void ClampFrameIndex();
 
 	bool IsReadyForTraining() const;
 	FText GetOverlayText();
