[MLDeformer]
- Support multiple meshes.
- Fixed several bugs (including some crash).
- Some code refactors.
- Disable the shrinkage related properties when Shrinkage isn't the selected loss function.

18131918

diff --git a/Engine/Plugins/Experimental/MLDeformer/Source/MLDeformer/Private/MLDeformerAsset.cpp b/Engine/Plugins/Experimental/MLDeformer/Source/MLDeformer/Private/MLDeformerAsset.cpp
--- a/Engine/Plugins/Experimental/MLDeformer/Source/MLDeformer/Private/MLDeformerAsset.cpp
+++ b/Engine/Plugins/Experimental/MLDeformer/Source/MLDeformer/Private/MLDeformerAsset.cpp
@@ -6,6 +6,7 @@
 #include "MLDeformerInputInfo.h"
 #include "GeometryCache.h"
 #include "GeometryCacheMeshData.h"
+#include "GeometryCacheTrack.h"
 #include "Rendering/SkeletalMeshModel.h"
 #include "Rendering/SkeletalMeshLODModel.h"
 #include "NeuralNetwork.h"
@@ -87,6 +88,110 @@ void UMLDeformerAsset::PostLoad()
 }
 
 #if WITH_EDITOR
+// A fuzzy name match.
+// There is a match when the track name starts with the mesh name.
+bool IsPotentialMatch(const FString& TrackName, const FString& MeshName)
+{
+	return (TrackName.Find(MeshName) == 0);
+}
+
+void UMLDeformerAsset::GenerateMeshMappings(UMLDeformerAsset* DeformerAsset, TArray<FMLDeformerMeshMapping>& OutMeshMappings, TArray<FString>& OutFailedImportedMeshNames)
+{
+	OutMeshMappings.Empty();
+	OutFailedImportedMeshNames.Empty();
+
+	USkeletalMesh* SkelMesh = DeformerAsset->GetSkeletalMesh();
+	UGeometryCache* GeomCache = DeformerAsset->GetGeometryCache();
+	if (SkelMesh == nullptr || GeomCache == nullptr)
+	{
+		return;
+	}
+
+	// If we haven't got any imported mesh infos then the asset needs to be reimported first.
+	// We show an error for this in the editor UI already.
+	FSkeletalMeshModel* ImportedModel = SkelMesh->GetImportedModel();
+	check(ImportedModel);
+	const TArray<FSkelMeshImportedMeshInfo>& SkelMeshInfos = ImportedModel->LODModels[0].ImportedMeshInfos;
+	if (SkelMeshInfos.IsEmpty())
+	{
+		return;
+	}
+
+	// For all meshes in the skeletal mesh.
+	const float SampleTime = 0.0f;
+	FString SkelMeshName;
+	for (int32 SkelMeshIndex = 0; SkelMeshIndex < SkelMeshInfos.Num(); ++SkelMeshIndex)
+	{
+		const FSkelMeshImportedMeshInfo& MeshInfo = SkelMeshInfos[SkelMeshIndex];
+		SkelMeshName = MeshInfo.Name.ToString();
+
+		// Find the matching one in the geom cache.
+		bool bFoundMatch = false;
+		for (int32 TrackIndex = 0; TrackIndex < GeomCache->Tracks.Num(); ++TrackIndex)
+		{
+			// Check if this is a candidate based on the mesh and track name.
+			UGeometryCacheTrack* Track = GeomCache->Tracks[TrackIndex];
+			if (Track && IsPotentialMatch(Track->GetName(), SkelMeshName))
+			{	
+				// Extract the geom cache mesh data.
+				FGeometryCacheMeshData GeomCacheMeshData;
+				if (!Track->GetMeshDataAtTime(SampleTime, GeomCacheMeshData))
+				{
+					continue;
+				}
+
+				// Verify that we have imported vertex numbers.
+				if (GeomCacheMeshData.ImportedVertexNumbers.IsEmpty())
+				{
+					continue;
+				}
+
+				// Get the number of geometry cache mesh imported verts.
+				int32 NumGeomMeshVerts = 0;
+				for (int32 GeomVertIndex = 0; GeomVertIndex < GeomCacheMeshData.ImportedVertexNumbers.Num(); ++GeomVertIndex)
+				{
+					NumGeomMeshVerts = FMath::Max(NumGeomMeshVerts, (int32)GeomCacheMeshData.ImportedVertexNumbers[GeomVertIndex]);
+				}
+				NumGeomMeshVerts += 1;	// +1 Because we use indices, so a cube's max index is 7, while there are 8 vertices.
+
+				// Make sure the vertex counts match.
+				const int32 NumSkelMeshVerts = MeshInfo.NumVertices;
+				if (NumSkelMeshVerts != NumGeomMeshVerts)
+				{
+					continue;
+				}
+
+				// Create a new mesh mapping entry.
+				OutMeshMappings.AddDefaulted();
+				FMLDeformerMeshMapping& Mapping = OutMeshMappings.Last();
+				Mapping.MeshIndex = SkelMeshIndex;
+				Mapping.TrackIndex = TrackIndex;
+				Mapping.SkelMeshToTrackVertexMap.AddUninitialized(NumSkelMeshVerts);
+
+				// For all vertices (both skel mesh and geom cache mesh have the same number of verts here).
+				for (int32 VertexIndex = 0; VertexIndex < NumSkelMeshVerts; ++VertexIndex)
+				{
+					// Find the first vertex with the same dcc vertex in the geom cache mesh.
+					// When there are multiple vertices with the same vertex number here, they are duplicates with different normals or uvs etc.
+					// However they all share the same vertex position, so we can just find the first hit, as we only need the position later on.
+					const int32 GeomCacheVertexIndex = GeomCacheMeshData.ImportedVertexNumbers.Find(VertexIndex);
+					Mapping.SkelMeshToTrackVertexMap[VertexIndex] = GeomCacheVertexIndex;
+				}
+
+				// We found a match, no need to iterate over more Tracks.
+				bFoundMatch = true;
+				break;
+			} // If the track name matches the skeletal meshes internal mesh name.
+		} // For all tracks.
+
+		if (!bFoundMatch)
+		{
+			OutFailedImportedMeshNames.Add(SkelMeshName);
+			UE_LOG(LogMLDeformer, Warning, TEXT("Imported mesh '%s' cannot be matched with a geometry cache track."), *SkelMeshName);
+		}
+	} // For all meshes inside the skeletal mesh.
+}
+
 FText UMLDeformerAsset::GetBaseAssetChangedErrorText() const
 {
 	FText Result;
@@ -125,6 +230,25 @@ FText UMLDeformerAsset::GetTargetAssetChangedErrorText() const
 	return Result;
 }
 
+FText UMLDeformerAsset::GetSkeletalMeshNeedsReimportErrorText() const
+{
+	FText Result;
+
+	if (SkeletalMesh)
+	{
+		FSkeletalMeshModel* ImportedModel = SkeletalMesh->GetImportedModel();
+		check(ImportedModel);
+
+		const TArray<FSkelMeshImportedMeshInfo>& SkelMeshInfos = ImportedModel->LODModels[0].ImportedMeshInfos;
+		if (SkelMeshInfos.IsEmpty())
+		{
+			Result = LOCTEXT("SkelMeshNeedsReimport", "Skeletal Mesh asset needs to be reimported.");
+		}
+	}
+
+	return Result;
+}
+
 void UMLDeformerAsset::PostEditChangeProperty(FPropertyChangedEvent& PropertyChangedEvent)
 {
 	static const FName SkeletalMeshName = GET_MEMBER_NAME_CHECKED(UMLDeformerAsset, SkeletalMesh);
@@ -140,7 +264,11 @@ void UMLDeformerAsset::PostEditChangeProperty(FPropertyChangedEvent& PropertyCha
 
 int32 UMLDeformerAsset::GetNumFrames() const
 {
-	check(GeometryCache != nullptr)
+	if (GeometryCache == nullptr)
+	{
+		return 0;
+	}
+
 	const int32 StartFrame = GeometryCache->GetStartFrame();
 	const int32 EndFrame = GeometryCache->GetEndFrame();
 	check(EndFrame >= StartFrame);
@@ -179,60 +307,85 @@ FText UMLDeformerAsset::GetGeomCacheErrorText(UGeometryCache* InGeomCache) const
 	FText Result;
 	if (InGeomCache)
 	{
+		FString ErrorString;
+
 		// Verify that we have imported vertex numbers enabled.
 		TArray<FGeometryCacheMeshData> MeshData;
 		InGeomCache->GetMeshDataAtTime(0.0f, MeshData);
 		if (MeshData.Num() == 0)
 		{
-			Result = LOCTEXT("TargetMeshNoMeshData", "No geometry data is present.");
+			ErrorString = FText(LOCTEXT("TargetMeshNoMeshData", "No geometry data is present.")).ToString();
 		}
 		else
 		{
 			if (MeshData[0].ImportedVertexNumbers.Num() == 0)
 			{
-				Result = LOCTEXT("TargetMeshNoImportedVertexNumbers", "Please import Geometry Cache with option 'Store Imported Vertex Numbers' enabled!");
+				ErrorString = FText(LOCTEXT("TargetMeshNoImportedVertexNumbers", "Please import Geometry Cache with option 'Store Imported Vertex Numbers' enabled!")).ToString();
 			}
 		}
+
+		// Check if we flattened the tracks.
+		if (InGeomCache->Tracks.Num() == 1 && InGeomCache->Tracks[0]->GetName() == TEXT("Flattened_Track"))
+		{
+			if (!ErrorString.IsEmpty())
+			{
+				ErrorString += TEXT("\n\n");
+			}
+			ErrorString += FText(LOCTEXT("TargetMeshFlattened", "Please import Geometry Cache with option 'Flatten Tracks' disabled!")).ToString();
+		}
+
+		Result = FText::FromString(ErrorString);
 	}
 
 	return Result;
 }
 
+FText UMLDeformerAsset::GetMeshMappingErrorText() const
+{
+	FText Result;
+	if (GeometryCache && SkeletalMesh)
+	{
+		// Check for failed mesh mappings.
+		UMLDeformerAsset* Asset = const_cast<UMLDeformerAsset*>(this);
+		TArray<FMLDeformerMeshMapping> MeshMappings;
+		TArray<FString> FailedNames;
+		UMLDeformerAsset::GenerateMeshMappings(Asset, MeshMappings, FailedNames);
+
+		// List all mesh names that have issues.
+		FString ErrorString;
+		for (int32 Index = 0; Index < FailedNames.Num(); ++Index)
+		{
+			ErrorString += FailedNames[Index];
+			if (Index < FailedNames.Num() - 1)
+			{
+				ErrorString += TEXT("\n");
+			}
+		}
+
+		Result = FText::FromString(ErrorString);
+	}
+	return Result;
+}
+
 FText UMLDeformerAsset::GetVertexErrorText(USkeletalMesh* InSkelMesh, UGeometryCache* InGeomCache, const FText& SkelName, const FText& GeomCacheName) const
 {
 	FText Result;
 
 	if (InSkelMesh && InGeomCache)
 	{
-		FSkeletalMeshModel* SkeletalMeshModel = InSkelMesh->GetImportedModel();	
-		TArray<FGeometryCacheMeshData> MeshData;
-		InGeomCache->GetMeshDataAtTime(0.0f, MeshData);
-		if (SkeletalMeshModel && MeshData.Num() > 0)
+		const int32 SkelVertCount = UMLDeformerAsset::ExtractNumImportedSkinnedVertices(InSkelMesh);
+		const int32 GeomCacheVertCount = UMLDeformerAsset::ExtractNumImportedGeomCacheVertices(InGeomCache);
+		const bool bHasGeomCacheError = !UMLDeformerAsset::GetGeomCacheErrorText(InGeomCache).IsEmpty();
+		if (SkelVertCount != GeomCacheVertCount && !bHasGeomCacheError)
 		{
-			const TArray<int32>& SkeletalVertexMap = SkeletalMeshModel->LODModels[0].MeshToImportVertexMap;
-			const TArray<uint32>& GeomCacheVertexMap = MeshData[0].ImportedVertexNumbers;
-			if (!GeomCacheVertexMap.IsEmpty() && !SkeletalVertexMap.IsEmpty())
-			{
-				// Find the max vertex index.
-				int32 MaxGeomCacheVertex = 0;
-				for (int32 Index = 0; Index < GeomCacheVertexMap.Num(); ++Index)
-				{
-					MaxGeomCacheVertex = FMath::Max<int32>(static_cast<int32>(GeomCacheVertexMap[Index]), MaxGeomCacheVertex);
-				}
-
-				const int32 MaxSkelMeshVertex = SkeletalMeshModel->LODModels[0].MaxImportVertex;
-				if (MaxSkelMeshVertex != MaxGeomCacheVertex)
-				{
-					Result = FText::Format(
-						LOCTEXT("MeshVertexNumVertsMismatch", "Vertex count of {0} doesn't match with {1}!\n\n{2} has {3} verts, while {4} has {5} verts."),
-						SkelName,
-						GeomCacheName,
-						SkelName,
-						MaxSkelMeshVertex + 1,
-						GeomCacheName,
-						MaxGeomCacheVertex + 1);
-				}
-			}
+			Result = FText::Format(
+				LOCTEXT("MeshVertexNumVertsMismatch", "Vertex count of {0} doesn't match with {1}!\n\n{2} has {3} verts, while {4} has {5} verts."),
+				SkelName,
+				GeomCacheName,
+				SkelName,
+				SkelVertCount,
+				GeomCacheName,
+				GeomCacheVertCount);
 		}
 	}
 
@@ -301,40 +454,8 @@ bool UMLDeformerAsset::IsCompatibleWithNeuralNet() const
 
 void UMLDeformerAsset::UpdateCachedNumVertices()
 {
-	NumSkeletalMeshVerts = 0;
-	NumGeomCacheVerts = 0;
-
-	// Extract max num vertices from the geometry cache.
-	if (GeometryCache)
-	{
-		TArray<FGeometryCacheMeshData> MeshData;
-		GeometryCache->GetMeshDataAtTime(0.0f, MeshData);
-		if (MeshData.Num() > 0)
-		{
-			const TArray<uint32>& VertexNumbers = MeshData[0].ImportedVertexNumbers;
-			for (int32 Index = 0; Index < VertexNumbers.Num(); ++Index)
-			{
-				NumGeomCacheVerts = FMath::Max<int32>(static_cast<int32>(VertexNumbers[Index]), NumGeomCacheVerts);
-			}
-
-			// Add one as they are indices starting from 0.
-			if (VertexNumbers.Num() > 0)
-			{
-				NumGeomCacheVerts += 1;
-			}
-		}
-	}
-
-	// Extract max num vertices from the skeletal mesh.
-	if (SkeletalMesh)
-	{
-		FSkeletalMeshModel* SkeletalMeshModel = SkeletalMesh->GetImportedModel();
-		if (SkeletalMeshModel)
-		{
-			const int32 MaxIndex = SkeletalMeshModel->LODModels[0].MaxImportVertex;
-			NumSkeletalMeshVerts = (MaxIndex > 0) ? (MaxIndex + 1) : 0;
-		}
-	}
+	NumSkeletalMeshVerts = UMLDeformerAsset::ExtractNumImportedSkinnedVertices(SkeletalMesh);
+	NumGeomCacheVerts = UMLDeformerAsset::ExtractNumImportedGeomCacheVertices(GeometryCache);
 }
 
 FMLDeformerInputInfo UMLDeformerAsset::CreateInputInfo() const
@@ -351,6 +472,44 @@ FMLDeformerInputInfo UMLDeformerAsset::CreateInputInfo() const
 
 	return Result;
 }
+
+int32 UMLDeformerAsset::ExtractNumImportedSkinnedVertices(USkeletalMesh* SkeletalMesh)
+{
+	return SkeletalMesh ? SkeletalMesh->GetNumImportedVertices() : 0;
+}
+
+int32 UMLDeformerAsset::ExtractNumImportedGeomCacheVertices(UGeometryCache* GeomCache)
+{
+	if (GeomCache == nullptr)
+	{
+		return 0;
+	}
+
+	int32 NumGeomCacheImportedVerts = 0;
+
+	// Extract the geom cache number of imported vertices.
+	TArray<FGeometryCacheMeshData> MeshDatas;
+	GeomCache->GetMeshDataAtTime(0.0f, MeshDatas);
+	for (const FGeometryCacheMeshData& MeshData : MeshDatas)
+	{
+		const TArray<uint32>& ImportedVertexNumbers = MeshData.ImportedVertexNumbers;
+		if (ImportedVertexNumbers.Num() > 0)
+		{
+			// Find the maximum value.
+			int32 MaxIndex = -1;
+			for (int32 Index = 0; Index < ImportedVertexNumbers.Num(); ++Index)
+			{
+				MaxIndex = FMath::Max(static_cast<int32>(ImportedVertexNumbers[Index]), MaxIndex);
+			}
+			check(MaxIndex > -1);
+
+			NumGeomCacheImportedVerts += MaxIndex + 1;
+		}
+	}
+
+	return NumGeomCacheImportedVerts;
+}
+
 #endif
 
 void UMLDeformerAsset::InitGPUData()
diff --git a/Engine/Plugins/Experimental/MLDeformer/Source/MLDeformer/Private/MLDeformerComponent.cpp b/Engine/Plugins/Experimental/MLDeformer/Source/MLDeformer/Private/MLDeformerComponent.cpp
--- a/Engine/Plugins/Experimental/MLDeformer/Source/MLDeformer/Private/MLDeformerComponent.cpp
+++ b/Engine/Plugins/Experimental/MLDeformer/Source/MLDeformer/Private/MLDeformerComponent.cpp
@@ -19,6 +19,7 @@ void UMLDeformerComponent::SetupComponent(UMLDeformerAsset* InDeformerAsset, USk
 {
 	DeformerAsset = InDeformerAsset;
 	SkelMeshComponent = InSkelMeshComponent;
+	AddTickPrerequisiteComponent(SkelMeshComponent);
 	DeformerInstance.Init(InDeformerAsset, InSkelMeshComponent);
 }
 
diff --git a/Engine/Plugins/Experimental/MLDeformer/Source/MLDeformer/Private/MLDeformerInputInfo.cpp b/Engine/Plugins/Experimental/MLDeformer/Source/MLDeformer/Private/MLDeformerInputInfo.cpp
--- a/Engine/Plugins/Experimental/MLDeformer/Source/MLDeformer/Private/MLDeformerInputInfo.cpp
+++ b/Engine/Plugins/Experimental/MLDeformer/Source/MLDeformer/Private/MLDeformerInputInfo.cpp
@@ -1,6 +1,7 @@
 // Copyright Epic Games, Inc. All Rights Reserved.
 
 #include "MLDeformerInputInfo.h"
+#include "MLDeformerAsset.h"
 #include "Engine/SkeletalMesh.h"
 #include "Animation/AnimCurveTypes.h"
 #include "Animation/Skeleton.h"
@@ -26,38 +27,8 @@ void FMLDeformerInputInfo::Init(const FMLDeformerInputInfoInitSettings& Settings
 	UGeometryCache* GeomCache = Settings.TargetMesh;
 
 #if WITH_EDITOR
-	if (SkeletalMesh)
-	{
-		// Get the number of imported vertices of the skeletal mesh.
-		FSkeletalMeshModel* SkeletalMeshModel = SkeletalMesh->GetImportedModel();	
-		if (SkeletalMeshModel)
-		{
-			if (!SkeletalMeshModel->LODModels[0].MeshToImportVertexMap.IsEmpty())
-			{
-				NumBaseMeshVertices = SkeletalMeshModel->LODModels[0].MaxImportVertex + 1;
-			}
-		}
-	}
-
-	if (GeomCache)
-	{
-		// Get the number of imported vertices of the geometry cache.
-		TArray<FGeometryCacheMeshData> MeshData;
-		GeomCache->GetMeshDataAtTime(0.0f, MeshData);
-		if (MeshData.Num() > 0)
-		{
-			int32 MaxGeomCacheVertex = 0;
-			const TArray<uint32>& GeomCacheVertexMap = MeshData[0].ImportedVertexNumbers;
-			if (!GeomCacheVertexMap.IsEmpty())
-			{
-				for (int32 Index = 0; Index < GeomCacheVertexMap.Num(); ++Index)
-				{
-					MaxGeomCacheVertex = FMath::Max<int32>(static_cast<int32>(GeomCacheVertexMap[Index]), MaxGeomCacheVertex);
-				}
-				NumTargetMeshVertices = MaxGeomCacheVertex + 1;
-			}
-		}
-	}
+	NumBaseMeshVertices = UMLDeformerAsset::ExtractNumImportedSkinnedVertices(SkeletalMesh);
+	NumTargetMeshVertices = UMLDeformerAsset::ExtractNumImportedGeomCacheVertices(GeomCache);
 #endif
 
 	// Handle bones.
@@ -230,6 +201,20 @@ FString FMLDeformerInputInfo::GenerateCompatibilityErrorString(USkeletalMesh* Sk
 		}
 	}
 
+	// Check vertex count.
+#if WITH_EDITORONLY_DATA
+	if ((NumBaseMeshVertices > 0 && NumTargetMeshVertices > 0) &&
+		NumBaseMeshVertices != SkeletalMesh->GetNumImportedVertices())
+	{
+		ErrorString += FString::Format(TEXT("The number of vertices that the network was trained on ({0} verts) doesn't match the skeletal mesh '{1}' ({2} verts)..\n"), 
+			{
+				NumBaseMeshVertices, 
+				SkeletalMesh->GetName(),
+				SkeletalMesh->GetNumImportedVertices(),
+			} );
+	}
+#endif
+
 	return ErrorString;
 }
 
diff --git a/Engine/Plugins/Experimental/MLDeformer/Source/MLDeformer/Public/MLDeformerAsset.h b/Engine/Plugins/Experimental/MLDeformer/Source/MLDeformer/Public/MLDeformerAsset.h
--- a/Engine/Plugins/Experimental/MLDeformer/Source/MLDeformer/Public/MLDeformerAsset.h
+++ b/Engine/Plugins/Experimental/MLDeformer/Source/MLDeformer/Public/MLDeformerAsset.h
@@ -75,6 +75,15 @@ private:
 	TArray<int32> VertexMap;
 };
 
+#if WITH_EDITORONLY_DATA
+struct MLDEFORMER_API FMLDeformerMeshMapping
+{
+	int32 MeshIndex = INDEX_NONE;	// The imported model's mesh info index.
+	int32 TrackIndex = INDEX_NONE;	// The geom cache track that this mesh is mapped to.
+	TArray<int32> SkelMeshToTrackVertexMap;	// This maps imported model individual meshes to the geomcache track's mesh data.
+};
+#endif
+
 /**
  * The machine learning deformer asset.
  * At runtime this contains only the data needed to run the neural network inference.
@@ -106,6 +115,12 @@ public:
 	FText GetTargetAssetChangedErrorText() const;
 	FText GetInputsErrorText() const;
 	FText GetIncompatibleSkeletonErrorText(USkeletalMesh* InSkelMesh, UAnimSequence* InAnimSeq) const;
+	FText GetSkeletalMeshNeedsReimportErrorText() const;
+	FText GetMeshMappingErrorText() const;
+
+	static int32 ExtractNumImportedGeomCacheVertices(UGeometryCache* GeomCache);
+	static int32 ExtractNumImportedSkinnedVertices(USkeletalMesh* SkeletalMesh);
+	static void GenerateMeshMappings(UMLDeformerAsset* DeformerAsset, TArray<FMLDeformerMeshMapping>& OutMeshMappings, TArray<FString>& OutFailedImportedMeshNames);
 
 	void UpdateCachedNumVertices();
 	bool IsCompatibleWithNeuralNet() const;
@@ -255,11 +270,11 @@ public:
 	ELossFunction LossFunction = ELossFunction::L1;
 
 	/** Shrinkage speed. Only if the shrinkage loss is used. */
-	UPROPERTY(EditAnywhere, BlueprintReadWrite, AdvancedDisplay, Category = "Training Settings", meta = (ClampMin = "1.0", ClampMax = "1000.0"))
+	UPROPERTY(EditAnywhere, BlueprintReadWrite, AdvancedDisplay, Category = "Training Settings", meta = (EditCondition = "LossFunction == ELossFunction::Shrinkage", ClampMin = "1.0", ClampMax = "1000.0"))
 	float ShrinkageSpeed = 10.0f;
 
 	/** Shrinkage threshold. Only if the shrinkage loss is used. */
-	UPROPERTY(EditAnywhere, BlueprintReadWrite, AdvancedDisplay, Category = "Training Settings", meta = (ClampMin = "0.0"))
+	UPROPERTY(EditAnywhere, BlueprintReadWrite, AdvancedDisplay, Category = "Training Settings", meta = (EditCondition = "LossFunction == ELossFunction::Shrinkage", ClampMin = "0.0"))
 	float ShrinkageThreshold = 0.1f;
 
 	/** The percentage of noise to add. */
diff --git a/Engine/Plugins/Experimental/MLDeformer/Source/MLDeformer/Public/MLDeformerVizSettings.h b/Engine/Plugins/Experimental/MLDeformer/Source/MLDeformer/Public/MLDeformerVizSettings.h
--- a/Engine/Plugins/Experimental/MLDeformer/Source/MLDeformer/Public/MLDeformerVizSettings.h
+++ b/Engine/Plugins/Experimental/MLDeformer/Source/MLDeformer/Public/MLDeformerVizSettings.h
@@ -41,8 +41,6 @@ public:
 	float GetAnimPlaySpeed() const { return AnimPlaySpeed; }
 	float GetVertexDeltaMultiplier() const { return VertexDeltaMultiplier; }
 	int32 GetFrameNumber() const { return FrameNumber; }
-	int32 GetVertexNumber() const { return VertexNumber; }
-	bool GetDrawVertex() const { return bDrawVertex; }
 	bool GetDrawLabels() const { return bDrawLabels; }
 	float GetLabelScale() const { return LabelScale; }
 	bool GetShowHeatMap() const { return bShowHeatMap; }
@@ -106,14 +104,6 @@ public:
 	UPROPERTY(EditAnywhere, Category = "Training Meshes", meta = (ClampMin = "0"))
 	uint32 FrameNumber = 0;
 
-	/** The vertex number to visualize. */
-	UPROPERTY(EditAnywhere, Category = "Training Meshes", meta = (ClampMin = "0"))
-	int32 VertexNumber = 0;
-
-	/** Draw the specified vertex? */
-	UPROPERTY(EditAnywhere, Category = "Training Meshes")
-	bool bDrawVertex = false;
-
 	/** Draw the vertex deltas? */
 	UPROPERTY(EditAnywhere, Category = "Training Meshes")
 	bool bDrawDeltas = true;
diff --git a/Engine/Plugins/Experimental/MLDeformer/Source/MLDeformerEditor/Private/MLDeformerAssetDetails.cpp b/Engine/Plugins/Experimental/MLDeformer/Source/MLDeformerEditor/Private/MLDeformerAssetDetails.cpp
--- a/Engine/Plugins/Experimental/MLDeformer/Source/MLDeformerEditor/Private/MLDeformerAssetDetails.cpp
+++ b/Engine/Plugins/Experimental/MLDeformer/Source/MLDeformerEditor/Private/MLDeformerAssetDetails.cpp
@@ -5,6 +5,8 @@
 #include "MLDeformerAsset.h"
 #include "NeuralNetwork.h"
 
+#include "GeometryCache.h"
+#include "GeometryCacheTrack.h"
 #include "Animation/AnimSequence.h"
 #include "Engine/SkeletalMesh.h"
 #include "PropertyCustomizationHelpers.h"
@@ -47,6 +49,7 @@ void FMLDeformerAssetDetails::CustomizeDetails(class IDetailLayoutBuilder& Detai
 	BaseMeshCategoryBuilder.AddProperty(GET_MEMBER_NAME_CHECKED(UMLDeformerAsset, SkeletalMesh));
 	if (DeformerAsset)
 	{
+		// Check if the base mesh matches the target mesh vertex count.
 		const FText ErrorText = DeformerAsset->GetVertexErrorText(DeformerAsset->SkeletalMesh, DeformerAsset->GeometryCache, FText::FromString("Base Mesh"), FText::FromString("Target Mesh"));
 		FDetailWidgetRow& ErrorRow = BaseMeshCategoryBuilder.AddCustomRow(FText::FromString("BaseMeshError"))
 			.Visibility(!ErrorText.IsEmpty() ? EVisibility::Visible : EVisibility::Collapsed)
@@ -61,6 +64,7 @@ void FMLDeformerAssetDetails::CustomizeDetails(class IDetailLayoutBuilder& Detai
 				]
 			];
 
+		// Check if the vertex counts of our asset have changed.
 		const FText ChangedErrorText = DeformerAsset->GetBaseAssetChangedErrorText();
 		FDetailWidgetRow& ChangedErrorRow = BaseMeshCategoryBuilder.AddCustomRow(FText::FromString("BaseMeshChangedError"))
 			.Visibility(!ChangedErrorText.IsEmpty() ? EVisibility::Visible : EVisibility::Collapsed)
@@ -74,6 +78,53 @@ void FMLDeformerAssetDetails::CustomizeDetails(class IDetailLayoutBuilder& Detai
 					.Message(ChangedErrorText)
 				]
 			];
+
+		// Check if our skeletal mesh's imported model contains a list of mesh infos. If not, we need to reimport it as it is an older asset.
+		const FText NeedsReimportErrorText = DeformerAsset->GetSkeletalMeshNeedsReimportErrorText();
+		FDetailWidgetRow& NeedsReimportErrorRow = BaseMeshCategoryBuilder.AddCustomRow(FText::FromString("BaseMeshNeedsReimportError"))
+			.Visibility(!NeedsReimportErrorText.IsEmpty() ? EVisibility::Visible : EVisibility::Collapsed)
+			.WholeRowContent()
+			[
+				SNew(SBox)
+				.Padding(FMargin(0.0f, 4.0f))
+				[
+					SNew(SWarningOrErrorBox)
+					.MessageStyle(EMessageStyle::Error)
+					.Message(NeedsReimportErrorText)
+				]
+			];
+		
+		// Show meshes that have no matching geometry track list.
+		const FText MeshMappingErrorList = DeformerAsset->GetMeshMappingErrorText();
+		FString GeomTrackNameList;
+		if (!MeshMappingErrorList.IsEmpty())
+		{
+			UGeometryCache* GeomCache = DeformerAsset->GetGeometryCache();
+			for (int32 Index = 0; Index < GeomCache->Tracks.Num(); ++Index)
+			{
+				GeomTrackNameList += GeomCache->Tracks[Index]->GetName();
+				if (Index < GeomCache->Tracks.Num() - 1)
+				{
+					GeomTrackNameList += TEXT("\n");
+				}
+			}
+		}
+		FText MeshMappingErrorFull = FText::Format(
+			LOCTEXT("MeshMappingError", "No matching GeomCache Tracks names found for meshes:\n{0}\n\nGeomCache Track List:\n{1}"), 
+			MeshMappingErrorList,
+			FText::FromString(GeomTrackNameList));
+		FDetailWidgetRow& MeshMappingErrorRow = BaseMeshCategoryBuilder.AddCustomRow(FText::FromString("MeshMappingError"))
+			.Visibility(!MeshMappingErrorList.IsEmpty() ? EVisibility::Visible : EVisibility::Collapsed)
+			.WholeRowContent()
+			[
+				SNew(SBox)
+				.Padding(FMargin(0.0f, 4.0f))
+				[
+					SNew(SWarningOrErrorBox)
+					.MessageStyle(EMessageStyle::Error)
+					.Message(MeshMappingErrorFull)
+				]
+			];
 	}
 
 	// Animation sequence.
diff --git a/Engine/Plugins/Experimental/MLDeformer/Source/MLDeformerEditor/Private/MLDeformerEditorData.cpp b/Engine/Plugins/Experimental/MLDeformer/Source/MLDeformerEditor/Private/MLDeformerEditorData.cpp
--- a/Engine/Plugins/Experimental/MLDeformer/Source/MLDeformerEditor/Private/MLDeformerEditorData.cpp
+++ b/Engine/Plugins/Experimental/MLDeformer/Source/MLDeformerEditor/Private/MLDeformerEditorData.cpp
@@ -188,38 +188,6 @@ float FMLDeformerEditorData::GetSnappedFrameTime(float InTime) const
 	return InTime;
 }
 
-int32 FMLDeformerEditorData::GetNumImportedVertices() const
-{
-	UGeometryCacheComponent* GeomCacheComponent = GetEditorActor(EMLDeformerEditorActorIndex::Target).GeomCacheComponent;
-	check(GeomCacheComponent);
-
-	UGeometryCache* GeomCache = GeomCacheComponent->GetGeometryCache();
-	if (GeomCache == nullptr)
-	{
-		return 0;
-	}
-
-	TArray<FGeometryCacheMeshData> MeshData;
-	GeomCache->GetMeshDataAtTime(GeomCacheComponent->GetAnimationTime(), MeshData);
-
-	if (MeshData.Num() > 0)
-	{
-		const TArray<uint32>& ImportedVertexNumbers = MeshData[0].ImportedVertexNumbers;
-		if (ImportedVertexNumbers.Num() > 0)
-		{
-			// Find the maximum value.
-			int32 MaxIndex = -1;
-			for (int32 i = 0; i < ImportedVertexNumbers.Num(); ++i)
-			{
-				MaxIndex = FMath::Max(static_cast<int32>(ImportedVertexNumbers[i]), MaxIndex);
-			}
-			return (MaxIndex + 1);
-		}
-		return 0;
-	}
-	return 0;
-}
-
 void FMLDeformerEditorData::UpdateTestAnimPlaySpeed()
 {
 	UMLDeformerVizSettings* VizSettings = GetDeformerAsset()->GetVizSettings();
@@ -368,6 +336,12 @@ bool FMLDeformerEditorData::IsReadyForTraining() const
 		return false;
 	}
 
+	// There are no training frames.
+	if (GetDeformerAsset()->GetNumFrames() == 0)
+	{
+		return false;
+	}
+
 	// Now make sure the assets are compatible.
 	if (!Asset->GetVertexErrorText(Asset->GetSkeletalMesh(), Asset->GetGeometryCache(), FText(), FText()).IsEmpty() ||
 		!Asset->GetGeomCacheErrorText(Asset->GetGeometryCache()).IsEmpty())
@@ -381,6 +355,12 @@ bool FMLDeformerEditorData::IsReadyForTraining() const
 		return false;
 	}
 
+	// Make sure every skeletal imported mesh has some geometry track.
+	if (!GetSingleFrameCache().GetSampler().GetFailedImportedMeshNames().IsEmpty())
+	{
+		return false;
+	}
+
 	return true;
 }
 
@@ -399,7 +379,6 @@ bool FMLDeformerEditorData::GenerateDeltas(uint32 LODIndex, uint32 FrameNumber,
 	// Let's also directly extract the positions of the vertices, since we already calculated them.
 	const FMLDeformerSamplerData& SamplerData = SingleFrameCache.GetSampler().GetSamplerData();
 	LinearSkinnedPositions = SamplerData.GetSkinnedVertexPositions();
-	GeomCachePositions = SamplerData.GetGeometryCachePositions();
 
 	return true;
 }
diff --git a/Engine/Plugins/Experimental/MLDeformer/Source/MLDeformerEditor/Private/MLDeformerEditorMode.cpp b/Engine/Plugins/Experimental/MLDeformer/Source/MLDeformerEditor/Private/MLDeformerEditorMode.cpp
--- a/Engine/Plugins/Experimental/MLDeformer/Source/MLDeformerEditor/Private/MLDeformerEditorMode.cpp
+++ b/Engine/Plugins/Experimental/MLDeformer/Source/MLDeformerEditor/Private/MLDeformerEditorMode.cpp
@@ -93,28 +93,6 @@ void FMLDeformerEditorMode::Render(const FSceneView* View, FViewport* Viewport,
 	const UMLDeformerVizSettings* VizSettings = DeformerAsset->GetVizSettings();
 	if (VizSettings->GetVisualizationMode() == EMLDeformerVizMode::TrainingData)
 	{
-		// Draw a single vertex.
-		if (VizSettings->GetDrawVertex())
-		{
-			const FLinearColor VertexColor = FMLDeformerEditorStyle::Get().GetColor("MLDeformer.Vertex.Color");
-			const float VertexSize = FMLDeformerEditorStyle::Get().GetFloat("MLDeformer.Vertex.Size");
-			const TArray<FVector3f>& SkinnedPositions = Data->LinearSkinnedPositions;
-			const FVector Offset = VizSettings->GetMeshSpacingOffsetVector();
-			const TArray<FVector3f>& GeomCachePositions = Data->GeomCachePositions;
-			const int32 VertexIndex = VizSettings->GetVertexNumber();
-			if (VertexIndex < SkinnedPositions.Num())
-			{
-				const FVector BasePos = SkinnedPositions[VertexIndex];
-				PDI->DrawPoint(BasePos, VertexColor, VertexSize, 10);
-			}
-
-			if (VertexIndex < GeomCachePositions.Num())
-			{
-				const FVector TargetPos = GeomCachePositions[VertexIndex] + Offset;
-				PDI->DrawPoint(TargetPos, VertexColor, VertexSize, 10);
-			}
-		}
-
 		// Draw deltas.
 		if (VizSettings->GetDrawVertexDeltas() && (Data->VertexDeltas.Num() / 3) == Data->LinearSkinnedPositions.Num())
 		{
@@ -242,7 +220,6 @@ void FMLDeformerEditorMode::UpdateActors()
 void FMLDeformerEditorMode::Tick(FEditorViewportClient* ViewportClient, float DeltaTime)
 {
 	FMLDeformerEditorData* Data = EditorData.Pin().Get();
-
 	Data->ClampFrameIndex();
 
 	UpdateActors();
@@ -250,9 +227,8 @@ void FMLDeformerEditorMode::Tick(FEditorViewportClient* ViewportClient, float De
 
 	// Calculate the training deltas when needed.
 	UMLDeformerAsset* DeformerAsset = Data->GetDeformerAsset();
-	const UMLDeformerVizSettings* VizSettings = DeformerAsset->GetVizSettings();
-	if ((VizSettings->GetVisualizationMode() == EMLDeformerVizMode::TrainingData) && 
-		(VizSettings->GetDrawVertex() || VizSettings->GetDrawVertexDeltas()))
+	UMLDeformerVizSettings* VizSettings = DeformerAsset->GetVizSettings();
+	if ((VizSettings->GetVisualizationMode() == EMLDeformerVizMode::TrainingData) && VizSettings->GetDrawVertexDeltas())
 	{
 		Data->GenerateDeltas(0, VizSettings->FrameNumber, Data->VertexDeltas);
 	}
diff --git a/Engine/Plugins/Experimental/MLDeformer/Source/MLDeformerEditor/Private/MLDeformerEditorToolkit.cpp b/Engine/Plugins/Experimental/MLDeformer/Source/MLDeformerEditor/Private/MLDeformerEditorToolkit.cpp
--- a/Engine/Plugins/Experimental/MLDeformer/Source/MLDeformerEditor/Private/MLDeformerEditorToolkit.cpp
+++ b/Engine/Plugins/Experimental/MLDeformer/Source/MLDeformerEditor/Private/MLDeformerEditorToolkit.cpp
@@ -186,7 +186,7 @@ void FMLDeformerEditorToolkit::FillToolbar(FToolBarBuilder& ToolbarBuilder)
 						const int32 NumFrames = (DeformerAsset->GetGeometryCache()->GetEndFrame() - DeformerAsset->GetGeometryCache()->GetStartFrame()) + 1;
 						FMLDeformerFrameCache::FInitSettings FrameCacheInitSettings;
 						FrameCacheInitSettings.DeformerAsset = DeformerAsset;
-						FrameCacheInitSettings.CacheSizeInBytes = 1024 * 1024 * DeformerAsset->GetCacheSizeInMegabytes();
+						FrameCacheInitSettings.CacheSizeInBytes = 1024ull * 1024 * DeformerAsset->GetCacheSizeInMegabytes();
 						FrameCacheInitSettings.World = EditorData->GetWorld();
 						TSharedPtr<FMLDeformerFrameCache> FrameCache = MakeShared<FMLDeformerFrameCache>();
 						FrameCache->Init(FrameCacheInitSettings);
diff --git a/Engine/Plugins/Experimental/MLDeformer/Source/MLDeformerEditor/Private/MLDeformerFrameCache.cpp b/Engine/Plugins/Experimental/MLDeformer/Source/MLDeformerEditor/Private/MLDeformerFrameCache.cpp
--- a/Engine/Plugins/Experimental/MLDeformer/Source/MLDeformerEditor/Private/MLDeformerFrameCache.cpp
+++ b/Engine/Plugins/Experimental/MLDeformer/Source/MLDeformerEditor/Private/MLDeformerFrameCache.cpp
@@ -11,6 +11,7 @@
 #include "GeometryCacheComponent.h"
 #include "GeometryCache.h"
 #include "GeometryCacheMeshData.h"
+#include "GeometryCacheTrack.h"
 
 #include "Math/NumericLimits.h"
 
@@ -61,9 +62,6 @@ void FMLDeformerSamplerData::Init(const FInitSettings& InitSettings)
 	SkinnedVertexPositions.Empty();
 	SkinnedVertexPositions.AddUninitialized(NumImportedVertices);
 
-	GeomCacheVertexPositions.Empty();
-	GeomCacheVertexPositions.AddUninitialized(NumImportedVertices);
-
 	VertexDeltas.Empty(); // Multiply by 3 as we store this as an array of floats instead of an FVector3f.
 
 	BoneMatrices.Empty();
@@ -102,37 +100,70 @@ void FMLDeformerSamplerData::Update(int32 InAnimFrameIndex)
 		// Tick the component.
 		GeometryCacheComponent->SetManualTick(true);
 		GeometryCacheComponent->TickAtThisTime(SampleTime, false, false, false);
-
-		// Grab the vertex positions.
-		ExtractGeomCachePositions(LODIndex, TempVertexPositions, GeomCacheVertexPositions);
 	}
 
 	// Calculate the vertex deltas.
 	const float DeltaCutoffLength = DeformerAsset.GetDeltaCutoffLength();
-	CalculateVertexDeltas(SkinnedVertexPositions, GeomCacheVertexPositions, DeltaCutoffLength, VertexDeltas);
+	CalculateVertexDeltas(SkinnedVertexPositions, DeltaCutoffLength, VertexDeltas);
 }
 
-void FMLDeformerSamplerData::CalculateVertexDeltas(const TArray<FVector3f>& SkinnedPositions, const TArray<FVector3f>& GeomCachePositions, float DeltaCutoffLength, TArray<float>& OutVertexDeltas)
+void FMLDeformerSamplerData::CalculateVertexDeltas(const TArray<FVector3f>& SkinnedPositions, float DeltaCutoffLength, TArray<float>& OutVertexDeltas) const
 {
-	check(SkinnedPositions.Num() == GeomCachePositions.Num());
 	const int32 NumVerts = SkinnedPositions.Num();
 	OutVertexDeltas.Reset(NumVerts * 3);
 	OutVertexDeltas.AddUninitialized(NumVerts * 3);
-	for (int32 VertexIndex = 0; VertexIndex < NumVerts; ++VertexIndex)
+
+	UMLDeformerAsset* DeformerAsset = Sampler->GetInitSettings().DeformerAsset;
+	USkeletalMesh* SkelMesh = DeformerAsset->GetSkeletalMesh();
+	UGeometryCache* GeomCache = DeformerAsset->GetGeometryCache();
+
+	const FTransform AlignmentTransform = DeformerAsset->GetAlignmentTransform();
+
+	FSkeletalMeshModel* ImportedModel = SkelMesh->GetImportedModel();
+	check(SkelMesh);
+	check(GeomCache);
+	check(ImportedModel);
+
+	// Reset all deltas to zero.
+	for (int32 Index = 0; Index < OutVertexDeltas.Num(); ++Index)
 	{
-		const FVector Delta = GeomCachePositions[VertexIndex] - SkinnedPositions[VertexIndex];
-		const int32 ArrayIndex = 3 * VertexIndex;
-		if (Delta.Length() < DeltaCutoffLength)
+		OutVertexDeltas[Index] = 0.0f;
+	}
+
+	// For all mesh mappings we found.
+	const float SampleTime = GeometryCacheComponent->GetTimeAtFrame(AnimFrameIndex);
+	const TArray<FSkelMeshImportedMeshInfo>& SkelMeshInfos = ImportedModel->LODModels[0].ImportedMeshInfos;
+	for (int32 MeshMappingIndex = 0; MeshMappingIndex < Sampler->GetNumMeshMappings(); ++MeshMappingIndex)
+	{
+		const FMLDeformerMeshMapping& MeshMapping = Sampler->GetMeshMapping(MeshMappingIndex);
+		const FSkelMeshImportedMeshInfo& MeshInfo = SkelMeshInfos[MeshMapping.MeshIndex];
+		UGeometryCacheTrack* Track = GeomCache->Tracks[MeshMapping.TrackIndex];
+
+		// Sample the mesh data of the geom cache.
+		FGeometryCacheMeshData GeomCacheMeshData;
+		if (!Track->GetMeshDataAtTime(SampleTime, GeomCacheMeshData))
 		{
-			OutVertexDeltas[ArrayIndex] = Delta.X;
-			OutVertexDeltas[ArrayIndex + 1] = Delta.Y;
-			OutVertexDeltas[ArrayIndex + 2] = Delta.Z;
+			continue;
 		}
-		else
+
+		// Calculate the vertex deltas.
+		for (int32 VertexIndex = 0; VertexIndex < MeshInfo.NumVertices; ++VertexIndex)
 		{
-			OutVertexDeltas[ArrayIndex] = 0.0f;
-			OutVertexDeltas[ArrayIndex + 1] = 0.0f;
-			OutVertexDeltas[ArrayIndex + 2] = 0.0f;
+			const int32 SkinnedVertexIndex = MeshInfo.StartImportedVertex + VertexIndex;
+			const int32 GeomCacheVertexIndex = MeshMapping.SkelMeshToTrackVertexMap[VertexIndex];
+			if (GeomCacheVertexIndex != INDEX_NONE)
+			{
+				const FVector3f SkinnedVertexPos = SkinnedPositions[SkinnedVertexIndex];
+				const FVector3f GeomCacheVertexPos = AlignmentTransform.TransformPosition(GeomCacheMeshData.Positions[GeomCacheVertexIndex]);
+				const FVector3f Delta = GeomCacheVertexPos - SkinnedVertexPos;
+				if (Delta.Length() < DeltaCutoffLength)
+				{
+					const int32 ArrayIndex = 3 * SkinnedVertexIndex;
+					OutVertexDeltas[ArrayIndex] = Delta.X;
+					OutVertexDeltas[ArrayIndex + 1] = Delta.Y;
+					OutVertexDeltas[ArrayIndex + 2] = Delta.Z;
+				}
+			}
 		}
 	}
 }
@@ -172,41 +203,6 @@ void FMLDeformerSamplerData::ExtractSkinnedPositions(int32 LODIndex, TArray<FMat
 	}
 }
 
-void FMLDeformerSamplerData::ExtractGeomCachePositions(int32 LODIndex, TArray<FVector3f>& TempPositions, TArray<FVector3f>& OutPositions) const
-{
-	OutPositions.Reset();
-	TempPositions.Reset();
-
-	UGeometryCache* GeomCache = GeometryCacheComponent->GetGeometryCache();
-	if (GeomCache == nullptr)
-	{
-		return;
-	}
-
-	const float Time = GeometryCacheComponent->GetAnimationTime();
-	TArray<FGeometryCacheMeshData> MeshData;
-	GeomCache->GetMeshDataAtTime(Time, MeshData);
-
-	if (MeshData.Num() > 0)
-	{
-		const TArray<uint32>& ImportedVertexNumbers = MeshData[0].ImportedVertexNumbers;
-		if (ImportedVertexNumbers.Num() > 0)
-		{
-			TempPositions = MeshData[0].Positions;
-			const FTransform& AlignmentTransform = Sampler->GetDeformerAsset().GetAlignmentTransform();
-
-			// Transform the positions with the alignment transform.
-			OutPositions.AddZeroed(NumImportedVertices);
-			for (int32 Index = 0; Index < TempPositions.Num(); ++Index)
-			{
-				const int32 ImportedVertex = static_cast<int32>(ImportedVertexNumbers[Index]);
-				const FVector Pos = TempPositions[Index];
-				OutPositions[ImportedVertex] = AlignmentTransform.TransformPosition(Pos);
-			}
-		}
-	}
-}
-
 int32 FMLDeformerSamplerData::GetNumBones() const
 {
 	const UMLDeformerAsset& DeformerAsset = Sampler->GetDeformerAsset();
@@ -219,7 +215,6 @@ SIZE_T FMLDeformerSamplerData::CalcMemUsageInBytes() const
 	SIZE_T NumBytes = 0;
 	NumBytes += SkinnedVertexPositions.GetAllocatedSize();
 	NumBytes += TempVertexPositions.GetAllocatedSize();
-	NumBytes += GeomCacheVertexPositions.GetAllocatedSize();
 	NumBytes += BoneMatrices.GetAllocatedSize();
 	NumBytes += VertexDeltas.GetAllocatedSize();
 	NumBytes += BoneRotations.GetAllocatedSize();
@@ -274,11 +269,8 @@ void FMLDeformerSampler::Init(const FInitSettings& InInitSettings)
 	UMLDeformerAsset* DeformerAsset = InitSettings.DeformerAsset;
 	UGeometryCache* GeomCache = DeformerAsset->GetGeometryCache();
 	USkeletalMesh* SkeletalMesh = DeformerAsset->GetSkeletalMesh();
-	const int32 NumSkeletalMeshImportedVerts = ExtractNumImportedSkinnedVertices(DeformerAsset->GetSkeletalMesh());
-	const int32 NumGeomCacheImportedVerts = ExtractNumImportedGeomCacheVertices(GeomCache);
-	check(NumSkeletalMeshImportedVerts > 0);
-	check(NumGeomCacheImportedVerts > 0);
-	check(NumSkeletalMeshImportedVerts == NumGeomCacheImportedVerts);
+	const int32 NumSkeletalMeshImportedVerts = UMLDeformerAsset::ExtractNumImportedSkinnedVertices(DeformerAsset->GetSkeletalMesh());
+	const int32 NumGeomCacheImportedVerts = UMLDeformerAsset::ExtractNumImportedGeomCacheVertices(GeomCache);
 
 	// Create and init the items.
 	FMLDeformerSamplerData::FInitSettings DataInitSettings;
@@ -312,6 +304,8 @@ void FMLDeformerSampler::Init(const FInitSettings& InInitSettings)
 	// Set the actor root components.
 	SkelMeshActor->SetRootComponent(SkelMeshComponent);
 	GeomCacheActor->SetRootComponent(GeomCacheComponent);
+
+	UMLDeformerAsset::GenerateMeshMappings(InitSettings.DeformerAsset, MeshMappings, FailedImportedMeshNames);
 }
 
 // Spawn an actor.
@@ -324,49 +318,6 @@ AActor* FMLDeformerSampler::CreateActor(UWorld* InWorld, const FName& Name) cons
 	return Actor;
 }
 
-int32 FMLDeformerSampler::ExtractNumImportedSkinnedVertices(USkeletalMesh* SkeletalMesh) const
-{
-	check(SkeletalMesh);
-
-	// Extract the number of skinned mesh vertices.
-	int32 NumSkeletalMeshImportedVerts = 0;
-	const FSkeletalMeshModel* SkeletalMeshModel = SkeletalMesh->GetImportedModel();
-	if (SkeletalMeshModel)
-	{
-		const int32 MaxIndex = SkeletalMeshModel->LODModels[0].MaxImportVertex;
-		NumSkeletalMeshImportedVerts = (MaxIndex > 0) ? (MaxIndex + 1) : 0;
-	}
-	return NumSkeletalMeshImportedVerts;
-}
-
-int32 FMLDeformerSampler::ExtractNumImportedGeomCacheVertices(UGeometryCache* GeomCache) const
-{
-	check(GeomCache);
-	int32 NumGeomCacheImportedVerts = 0;
-
-	// Extract the geom cache number of imported vertices.
-	TArray<FGeometryCacheMeshData> MeshData;
-	GeomCache->GetMeshDataAtTime(0.0f, MeshData);
-	if (MeshData.Num() > 0)
-	{
-		const TArray<uint32>& ImportedVertexNumbers = MeshData[0].ImportedVertexNumbers;
-		if (ImportedVertexNumbers.Num() > 0)
-		{
-			// Find the maximum value.
-			int32 MaxIndex = -1;
-			for (int32 Index = 0; Index < ImportedVertexNumbers.Num(); ++Index)
-			{
-				MaxIndex = FMath::Max(static_cast<int32>(ImportedVertexNumbers[Index]), MaxIndex);
-			}
-			check(MaxIndex > -1);
-
-			NumGeomCacheImportedVerts = MaxIndex + 1;
-		}
-	}
-
-	return NumGeomCacheImportedVerts;
-}
-
 int32 FMLDeformerSampler::GetNumVertices() const
 {
 	return SamplerData.GetNumImportedVertices();
@@ -384,14 +335,15 @@ int32 FMLDeformerSampler::GetNumCurves() const
 
 int32 FMLDeformerSampler::GetNumFrames() const
 {
-	UMLDeformerAsset* DeformerAsset = InitSettings.DeformerAsset;
-	UGeometryCache* GeomCache = DeformerAsset->GetGeometryCache();
-	return (GeomCache->GetEndFrame() - GeomCache->GetStartFrame()) + 1;
+	return InitSettings.DeformerAsset->GetNumFrames();
 }
 
 SIZE_T FMLDeformerSampler::CalcMemUsageInBytes() const
 {
-	return SamplerData.CalcMemUsageInBytes();
+	SIZE_T Result = 0;
+	Result += SamplerData.CalcMemUsageInBytes();
+	Result += MeshMappings.GetAllocatedSize();
+	return Result;
 }
 
 //----------------------------------------------------------------------------------------------------------
@@ -543,6 +495,10 @@ void FMLDeformerFrameCache::UpdateFrameMap()
 int32 FMLDeformerFrameCache::GenerateFrame(int32 AnimFrameIndex)
 {
 	check(AnimFrameIndex >= 0);
+	check(!FrameMap.IsEmpty());
+
+	// Make sure we're in a valid range.
+	AnimFrameIndex = FMath::Clamp(AnimFrameIndex, 0, FrameMap.Num() - 1);
 
 	// Let the sampler item generate the deltas.
 	FMLDeformerSamplerData& SamplerData = Sampler.GetSamplerData();
diff --git a/Engine/Plugins/Experimental/MLDeformer/Source/MLDeformerEditor/Private/MLDeformerFrameCache.h b/Engine/Plugins/Experimental/MLDeformer/Source/MLDeformerEditor/Private/MLDeformerFrameCache.h
--- a/Engine/Plugins/Experimental/MLDeformer/Source/MLDeformerEditor/Private/MLDeformerFrameCache.h
+++ b/Engine/Plugins/Experimental/MLDeformer/Source/MLDeformerEditor/Private/MLDeformerFrameCache.h
@@ -44,7 +44,6 @@ public:
 	UGeometryCacheComponent* GetGeometryCacheComponent() const { return GeometryCacheComponent; }
 
 	const TArray<FVector3f>& GetSkinnedVertexPositions() const { return SkinnedVertexPositions; }
-	const TArray<FVector3f>& GetGeometryCachePositions() const { return GeomCacheVertexPositions; }
 
 	const TArray<float>& GetVertexDeltas() const { return VertexDeltas; }
 	const TArray<float>& GetBoneRotations() const { return BoneRotations; }
@@ -57,8 +56,7 @@ public:
 
 protected:
 	void ExtractSkinnedPositions(int32 LODIndex, TArray<FMatrix44f>& InBoneMatrices, TArray<FVector3f>& TempPositions, TArray<FVector3f>& OutPositions) const;
-	void ExtractGeomCachePositions(int32 LODIndex, TArray<FVector3f>& TempPositions, TArray<FVector3f>& OutPositions) const;
-	void CalculateVertexDeltas(const TArray<FVector3f>& SkinnedPositions, const TArray<FVector3f>& GeomCachePositions, float DeltaCutoffLength, TArray<float>& OutVertexDeltas);
+	void CalculateVertexDeltas(const TArray<FVector3f>& SkinnedPositions, float DeltaCutoffLength, TArray<float>& OutVertexDeltas) const;
 
 protected:
 	FMLDeformerSampler* Sampler = nullptr;
@@ -66,7 +64,6 @@ protected:
 	TObjectPtr<UGeometryCacheComponent> GeometryCacheComponent = nullptr;
 	TArray<FVector3f> SkinnedVertexPositions;
 	TArray<FVector3f> TempVertexPositions;
-	TArray<FVector3f> GeomCacheVertexPositions;
 	TArray<FMatrix44f> BoneMatrices;
 	TArray<float> VertexDeltas;	// (NumImportedVerts * 3) -> xyz
 	TArray<float> BoneRotations; // (NumBones * 4) -> quat xyzw
@@ -83,7 +80,7 @@ class FMLDeformerSampler
 public:
 	struct FInitSettings
 	{
-		UWorld* World = nullptr;
+		TObjectPtr<UWorld> World = nullptr;
 		TObjectPtr<UMLDeformerAsset> DeformerAsset = nullptr;
 	};
 
@@ -104,6 +101,10 @@ public:
 	int32 GetNumBones() const;
 	int32 GetNumCurves() const;
 	int32 GetNumFrames() const;
+	int32 GetNumMeshMappings() const { return MeshMappings.Num(); }
+
+	const FMLDeformerMeshMapping& GetMeshMapping(int32 Index) const { return MeshMappings[Index]; }
+	const TArray<FString>& GetFailedImportedMeshNames() const { return FailedImportedMeshNames; }
 
 	const FInitSettings& GetInitSettings() const { return InitSettings; }
 	const UMLDeformerAsset& GetDeformerAsset() const { return *InitSettings.DeformerAsset.Get(); }
@@ -113,12 +114,12 @@ public:
 
 protected:
 	AActor* CreateActor(UWorld* InWorld, const FName& Name) const;
-	int32 ExtractNumImportedSkinnedVertices(USkeletalMesh* SkeletalMesh) const;
-	int32 ExtractNumImportedGeomCacheVertices(UGeometryCache* GeomCache) const;
 
 protected:
 	TObjectPtr<AActor> SkelMeshActor = nullptr;
 	TObjectPtr<AActor> GeomCacheActor = nullptr;
+	TArray<FMLDeformerMeshMapping> MeshMappings; // Maps skeletal meshes imported meshes to geometry tracks. 
+	TArray<FString> FailedImportedMeshNames; // Imported mesh names in the skeletal mesh for which no geom cache track could be found.
 	FMLDeformerSamplerData SamplerData;
 	FInitSettings InitSettings;
 };
diff --git a/Engine/Plugins/Experimental/MLDeformer/Source/MLDeformerEditor/Private/MLDeformerVizSettingsDetails.cpp b/Engine/Plugins/Experimental/MLDeformer/Source/MLDeformerEditor/Private/MLDeformerVizSettingsDetails.cpp
--- a/Engine/Plugins/Experimental/MLDeformer/Source/MLDeformerEditor/Private/MLDeformerVizSettingsDetails.cpp
+++ b/Engine/Plugins/Experimental/MLDeformer/Source/MLDeformerEditor/Private/MLDeformerVizSettingsDetails.cpp
@@ -222,8 +222,6 @@ void FMLDeformerVizSettingsDetails::CustomizeDetails(class IDetailLayoutBuilder&
 	IDetailCategoryBuilder& TrainingMeshesCategoryBuilder = DetailBuilder.EditCategory("Training Meshes", FText::GetEmpty(), ECategoryPriority::Important);
 	TrainingMeshesCategoryBuilder.SetCategoryVisibility(bShowTrainingData);
 	TrainingMeshesCategoryBuilder.AddProperty(GET_MEMBER_NAME_CHECKED(UMLDeformerVizSettings, FrameNumber));
-	TrainingMeshesCategoryBuilder.AddProperty(GET_MEMBER_NAME_CHECKED(UMLDeformerVizSettings, VertexNumber));
-	TrainingMeshesCategoryBuilder.AddProperty(GET_MEMBER_NAME_CHECKED(UMLDeformerVizSettings, bDrawVertex));
 	TrainingMeshesCategoryBuilder.AddProperty(GET_MEMBER_NAME_CHECKED(UMLDeformerVizSettings, bDrawDeltas));
 	TrainingMeshesCategoryBuilder.AddProperty(GET_MEMBER_NAME_CHECKED(UMLDeformerVizSettings, bXRayDeltas));
 }
diff --git a/Engine/Plugins/Experimental/MLDeformer/Source/MLDeformerEditor/Private/MLPytorchDataSetInterface.cpp b/Engine/Plugins/Experimental/MLDeformer/Source/MLDeformerEditor/Private/MLPytorchDataSetInterface.cpp
--- a/Engine/Plugins/Experimental/MLDeformer/Source/MLDeformerEditor/Private/MLPytorchDataSetInterface.cpp
+++ b/Engine/Plugins/Experimental/MLDeformer/Source/MLDeformerEditor/Private/MLPytorchDataSetInterface.cpp
@@ -59,8 +59,7 @@ int32 UMLPytorchDataSetInterface::GetNumberSampleCurves() const
 int32 UMLPytorchDataSetInterface::GetNumberSampleDeltas() const
 {
 	check(IsValid());
-	const FMLDeformerEditorData* Data = EditorData.Get();
-	return Data->GetNumImportedVertices();
+	return UMLDeformerAsset::ExtractNumImportedGeomCacheVertices(EditorData->GetDeformerAsset()->GetGeometryCache());
 }
 
 int32 UMLPytorchDataSetInterface::NumSamples() const
diff --git a/Engine/Plugins/Experimental/MLDeformer/Source/MLDeformerEditor/Public/MLDeformerEditorData.h b/Engine/Plugins/Experimental/MLDeformer/Source/MLDeformerEditor/Public/MLDeformerEditorData.h
--- a/Engine/Plugins/Experimental/MLDeformer/Source/MLDeformerEditor/Public/MLDeformerEditorData.h
+++ b/Engine/Plugins/Experimental/MLDeformer/Source/MLDeformerEditor/Public/MLDeformerEditorData.h
@@ -96,15 +96,14 @@ public:
 	float GetDuration() const;
 	float GetSnappedFrameTime(float InTime) const;
 	float GetTimeAtFrame(int32 FrameNumber) const;
-	int32 GetNumImportedVertices() const;
 
 	/**
 	 *Update mean and scale of vertex deltas.
 	 *
 	 * @param TrainingFrame The training frame that contains the vertex deltas.
-	 * @param InOutMeanVertexDelta Mean vertex delta
-	 * @param InOutVertexDeltaScale Vertex delta scale
-	 * @param InOutCount Count
+	 * @param InOutMeanVertexDelta Mean vertex delta.
+	 * @param InOutVertexDeltaScale Vertex delta scale.
+	 * @param InOutCount Count.
 	 */
 	static void UpdateVertexDeltaMeanAndScale(const FMLDeformerTrainingFrame& TrainingFrame, FVector3f& InOutMeanVertexDelta, FVector3f& InOutVertexDeltaScale, float& InOutCount);
 
@@ -129,10 +128,8 @@ public:
 	
 public:
 	TArray<FVector3f> LinearSkinnedPositions;
-	TArray<FVector3f> GeomCachePositions;
 	TArray<float> VertexDeltas;
 	bool bIsVertexDeltaNormalized = false;
-
 	int32 CurrentFrame = -1;
 
 private:
diff --git a/Engine/Source/Runtime/Engine/Classes/Engine/SkeletalMesh.h b/Engine/Source/Runtime/Engine/Classes/Engine/SkeletalMesh.h
--- a/Engine/Source/Runtime/Engine/Classes/Engine/SkeletalMesh.h
+++ b/Engine/Source/Runtime/Engine/Classes/Engine/SkeletalMesh.h
@@ -828,6 +828,10 @@ public:
 	// End USkeletalMeshEditorData public skeletalmesh API
 	//////////////////////////////////////////////////////////////////////////
 
+	/** Get the number of imported vertices. This returns 0 if GetImportedModel() returns a nullptr.
+	  * This is the number of vertices as they appear in the source asset, for example 8 for a cube. */
+	int32 GetNumImportedVertices() const;
+
 	/** Get the imported data for this skeletal mesh. */
 	FORCEINLINE FSkeletalMeshModel* GetImportedModel() const 
 	{
diff --git a/Engine/Source/Runtime/Engine/Private/SkeletalMesh.cpp b/Engine/Source/Runtime/Engine/Private/SkeletalMesh.cpp
--- a/Engine/Source/Runtime/Engine/Private/SkeletalMesh.cpp
+++ b/Engine/Source/Runtime/Engine/Private/SkeletalMesh.cpp
@@ -1011,6 +1011,20 @@ void USkeletalMesh::ReleaseResources()
 	}
 }
 
+#if WITH_EDITORONLY_DATA
+int32 USkeletalMesh::GetNumImportedVertices() const
+{
+	const FSkeletalMeshModel* SkeletalMeshModel = GetImportedModel();
+	if (SkeletalMeshModel)
+	{
+		const int32 MaxIndex = SkeletalMeshModel->LODModels[0].MaxImportVertex;
+		return (MaxIndex > 0) ? (MaxIndex + 1) : 0;
+	}
+
+	return 0;
+}
+#endif
+
 #if WITH_EDITORONLY_DATA
 static void AccumulateUVDensities(float* OutWeightedUVDensities, float* OutWeights, const FSkeletalMeshLODRenderData& LODData, const FSkelMeshRenderSection& Section)
 {
