[MLDeformer]
- Pre-skinning support.
- Some bug fixes.
- Users can now train on a subset of the number of training samples in the anim.
- New pre-trained Biceps and Cylinder asset.
- New default training settings.
- Added some console commands to enable some debug rendering.

diff --git a/Engine/Plugins/Experimental/MLDeformer/Source/MLDeformer/Private/MLDeformerAsset.cpp b/Engine/Plugins/Experimental/MLDeformer/Source/MLDeformer/Private/MLDeformerAsset.cpp
--- a/Engine/Plugins/Experimental/MLDeformer/Source/MLDeformer/Private/MLDeformerAsset.cpp
+++ b/Engine/Plugins/Experimental/MLDeformer/Source/MLDeformer/Private/MLDeformerAsset.cpp
@@ -131,7 +131,9 @@ void UMLDeformerAsset::GenerateMeshMappings(UMLDeformerAsset* DeformerAsset, TAr
 		{
 			// Check if this is a candidate based on the mesh and track name.
 			UGeometryCacheTrack* Track = GeomCache->Tracks[TrackIndex];
-			if (Track && IsPotentialMatch(Track->GetName(), SkelMeshName))
+			const bool bIsSoloMesh = (GeomCache->Tracks.Num() == 1 && SkelMeshInfos.Num() == 1);	// Do we just have one mesh and one track?
+			if (Track && 
+				(IsPotentialMatch(Track->GetName(), SkelMeshName) || bIsSoloMesh))
 			{	
 				// Extract the geom cache mesh data.
 				FGeometryCacheMeshData GeomCacheMeshData;
@@ -167,6 +169,7 @@ void UMLDeformerAsset::GenerateMeshMappings(UMLDeformerAsset* DeformerAsset, TAr
 				Mapping.MeshIndex = SkelMeshIndex;
 				Mapping.TrackIndex = TrackIndex;
 				Mapping.SkelMeshToTrackVertexMap.AddUninitialized(NumSkelMeshVerts);
+				Mapping.ImportedVertexToRenderVertexMap.AddUninitialized(NumSkelMeshVerts);
 
 				// For all vertices (both skel mesh and geom cache mesh have the same number of verts here).
 				for (int32 VertexIndex = 0; VertexIndex < NumSkelMeshVerts; ++VertexIndex)
@@ -176,6 +179,10 @@ void UMLDeformerAsset::GenerateMeshMappings(UMLDeformerAsset* DeformerAsset, TAr
 					// However they all share the same vertex position, so we can just find the first hit, as we only need the position later on.
 					const int32 GeomCacheVertexIndex = GeomCacheMeshData.ImportedVertexNumbers.Find(VertexIndex);
 					Mapping.SkelMeshToTrackVertexMap[VertexIndex] = GeomCacheVertexIndex;
+					
+					// Map the source asset vertex number to a render vertex. This is the first duplicate of that vertex.
+					const int32 RenderVertexIndex = ImportedModel->LODModels[0].MeshToImportVertexMap.Find(MeshInfo.StartImportedVertex + VertexIndex);
+					Mapping.ImportedVertexToRenderVertexMap[VertexIndex] = RenderVertexIndex;
 				}
 
 				// We found a match, no need to iterate over more Tracks.
@@ -302,6 +309,11 @@ void UMLDeformerAsset::InitVertexMap()
 	}
 }
 
+int32 UMLDeformerAsset::GetNumFramesForTraining() const 
+{ 
+	return FMath::Min(GetNumFrames(), GetTrainingFrameLimit()); 
+}
+
 FText UMLDeformerAsset::GetGeomCacheErrorText(UGeometryCache* InGeomCache) const
 {
 	FText Result;
@@ -327,11 +339,24 @@ FText UMLDeformerAsset::GetGeomCacheErrorText(UGeometryCache* InGeomCache) const
 		// Check if we flattened the tracks.
 		if (InGeomCache->Tracks.Num() == 1 && InGeomCache->Tracks[0]->GetName() == TEXT("Flattened_Track"))
 		{
-			if (!ErrorString.IsEmpty())
+			int32 NumSkelMeshes = 0;
+			if (SkeletalMesh)
 			{
-				ErrorString += TEXT("\n\n");
+				FSkeletalMeshModel* Model = SkeletalMesh->GetImportedModel();
+				if (Model)
+				{
+					NumSkelMeshes = Model->LODModels[0].ImportedMeshInfos.Num();		
+				}
+			}
+
+			if (NumSkelMeshes > 1)
+			{
+				if (!ErrorString.IsEmpty())
+				{
+					ErrorString += TEXT("\n\n");
+				}
+				ErrorString += FText(LOCTEXT("TargetMeshFlattened", "Please import Geometry Cache with option 'Flatten Tracks' disabled!")).ToString();
 			}
-			ErrorString += FText(LOCTEXT("TargetMeshFlattened", "Please import Geometry Cache with option 'Flatten Tracks' disabled!")).ToString();
 		}
 
 		Result = FText::FromString(ErrorString);
diff --git a/Engine/Plugins/Experimental/MLDeformer/Source/MLDeformer/Private/MLDeformerComponent.cpp b/Engine/Plugins/Experimental/MLDeformer/Source/MLDeformer/Private/MLDeformerComponent.cpp
--- a/Engine/Plugins/Experimental/MLDeformer/Source/MLDeformer/Private/MLDeformerComponent.cpp
+++ b/Engine/Plugins/Experimental/MLDeformer/Source/MLDeformer/Private/MLDeformerComponent.cpp
@@ -41,7 +41,11 @@ void UMLDeformerComponent::TickComponent(float DeltaTime, enum ELevelTick TickTy
 	if (TickType != ELevelTick::LEVELTICK_PauseTick)
 	{
 		// Update the deformer, which runs the inference.
-		DeformerInstance.Update();
+		// Only do this when the LOD level is 0.
+		if (SkelMeshComponent && SkelMeshComponent->GetPredictedLODLevel() == 0)
+		{
+			DeformerInstance.Update();
+		}
 	}
 
 	Super::TickComponent(DeltaTime, TickType, ThisTickFunction);
diff --git a/Engine/Plugins/Experimental/MLDeformer/Source/MLDeformer/Private/MLDeformerInputInfo.cpp b/Engine/Plugins/Experimental/MLDeformer/Source/MLDeformer/Private/MLDeformerInputInfo.cpp
--- a/Engine/Plugins/Experimental/MLDeformer/Source/MLDeformer/Private/MLDeformerInputInfo.cpp
+++ b/Engine/Plugins/Experimental/MLDeformer/Source/MLDeformer/Private/MLDeformerInputInfo.cpp
@@ -12,6 +12,7 @@
 #include "Rendering/SkeletalMeshModel.h"
 #include "Rendering/SkeletalMeshLODModel.h"
 
+
 void FMLDeformerInputInfo::Init(const FMLDeformerInputInfoInitSettings& Settings)
 {
 	// Reset things.
@@ -243,13 +244,23 @@ void FMLDeformerInputInfo::ExtractBoneRotations(USkeletalMeshComponent* SkelMesh
 	{
 		const FName BoneName = GetBoneName(Index);
 		const int32 SkelMeshBoneIndex = SkelMeshComponent->GetBoneIndex(BoneName);
-		check(SkelMeshBoneIndex != INDEX_NONE);
-		const int32 Offset = Index * 4;
-		const FQuat Rotation = BoneTransforms[SkelMeshBoneIndex].GetRotation();
-		OutRotations[Offset] = Rotation.X;
-		OutRotations[Offset+1] = Rotation.Y;
-		OutRotations[Offset+2] = Rotation.Z;
-		OutRotations[Offset+3] = Rotation.W;
+		if (SkelMeshBoneIndex != INDEX_NONE)
+		{
+			const int32 Offset = Index * 4;
+			const FQuat Rotation = BoneTransforms[SkelMeshBoneIndex].GetRotation();
+			OutRotations[Offset] = Rotation.X;
+			OutRotations[Offset+1] = Rotation.Y;
+			OutRotations[Offset+2] = Rotation.Z;
+			OutRotations[Offset+3] = Rotation.W;
+		}
+		else
+		{
+			const int32 Offset = Index * 4;
+			OutRotations[Offset] = 0.0f;	// Set to identity quat.
+			OutRotations[Offset+1] = 0.0f;
+			OutRotations[Offset+2] = 0.0f;
+			OutRotations[Offset+3] = 1.0f;
+		}
 	}
 }
 
diff --git a/Engine/Plugins/Experimental/MLDeformer/Source/MLDeformer/Public/MLDeformerAsset.h b/Engine/Plugins/Experimental/MLDeformer/Source/MLDeformer/Public/MLDeformerAsset.h
--- a/Engine/Plugins/Experimental/MLDeformer/Source/MLDeformer/Public/MLDeformerAsset.h
+++ b/Engine/Plugins/Experimental/MLDeformer/Source/MLDeformer/Public/MLDeformerAsset.h
@@ -52,6 +52,13 @@ enum class EDeviceType : uint8
 	GPU		/** Train using GPU.*/
 };
 
+UENUM()
+enum class EDeltaMode : uint8
+{
+	PreSkinning,	/** Apply the deltas before skinning. */
+	PostSkinning	/** Apply the deltas after skinning. */
+};
+
 /** The training inputs. Specifies what data to include in training. */
 UENUM()
 enum class ETrainingInputs : uint8
@@ -81,6 +88,7 @@ struct MLDEFORMER_API FMLDeformerMeshMapping
 	int32 MeshIndex = INDEX_NONE;	// The imported model's mesh info index.
 	int32 TrackIndex = INDEX_NONE;	// The geom cache track that this mesh is mapped to.
 	TArray<int32> SkelMeshToTrackVertexMap;	// This maps imported model individual meshes to the geomcache track's mesh data.
+	TArray<int32> ImportedVertexToRenderVertexMap; // Map the imported dcc vertex number to a render vertex. This is just one of the duplicates, which shares the same position.
 };
 #endif
 
@@ -143,6 +151,8 @@ public:
 	void SetInputInfo(const FMLDeformerInputInfo& Input) { InputInfo = Input; }
 
 #if WITH_EDITORONLY_DATA
+	EDeltaMode GetDeltaMode() const { return DeltaMode; }
+
 	const FTransform& GetAlignmentTransform() const { return AlignmentTransform; }
 
 	const USkeletalMesh* GetSkeletalMesh() const { return SkeletalMesh;  }
@@ -158,6 +168,8 @@ public:
 	UMLDeformerVizSettings* GetVizSettings() { return VizSettings; }
 
 	int32 GetNumFrames() const;
+	int32 GetTrainingFrameLimit() const { return MaxTrainingFrames; }
+	int32 GetNumFramesForTraining() const;
 	int32 GetCacheSizeInMegabytes() const { return CacheSizeInMegabytes; }
 	float GetDeltaCutoffLength() const { return DeltaCutoffLength; }
 #endif
@@ -235,7 +247,7 @@ public:
 
 	/** The number of frames per batch when training the model. */
 	UPROPERTY(EditAnywhere, BlueprintReadWrite, Category = "Training Settings", meta = (ClampMin = "1"))
-	int32 BatchSize = 32;
+	int32 BatchSize = 128;
 
 	/** The number of epochs to process without any decay. */
 	UPROPERTY(EditAnywhere, BlueprintReadWrite, Category = "Training Settings", meta = (ClampMin = "1"))
@@ -243,19 +255,23 @@ public:
 
 	/** The number of epochs to process that include a decay. */
 	UPROPERTY(EditAnywhere, BlueprintReadWrite, Category = "Training Settings", meta = (ClampMin = "1"))
-	int32 EpochsWithDecay = 10;
+	int32 EpochsWithDecay = 15;
+
+	/** The maximum numer of training frames (samples) to train on. Use this to train on a sub-section of your full training data. */
+	UPROPERTY(EditAnywhere, BlueprintReadWrite, Category = "Training Settings", meta = (ClampMin = "1"))
+	int32 MaxTrainingFrames = 1000000;
 
 	/** The maximum allowed size of the training cache in memory, in megabytes. So a value of 1024 would be one gigabyte. The larger the cache size the faster the training. */
 	UPROPERTY(EditAnywhere, BlueprintReadWrite, Category = "Training Settings", meta = (ClampMin = 0))
-	int32 CacheSizeInMegabytes = 2048;
+	int32 CacheSizeInMegabytes = 4096;	// 4 Gigabyte
 
 	/** The learning rate used during the model training. */
 	UPROPERTY(EditAnywhere, BlueprintReadWrite, AdvancedDisplay, Category = "Training Settings", meta = (ClampMin = "0.000001", ClampMax = "1.0"))
-	float LearningRate = 0.0005f;
+	float LearningRate = 0.00175f;
 
 	/** The decay function to adapt the learning rate. */
 	UPROPERTY(EditAnywhere, BlueprintReadWrite, AdvancedDisplay, Category = "Training Settings")
-	EDecayFunction DecayFunction = EDecayFunction::Linear;
+	EDecayFunction DecayFunction = EDecayFunction::Multiplicative;
 
 	/** The decay rate to apply to the learning rate once non-decay epochs have been reached. Higher values give less decay. */
 	UPROPERTY(EditAnywhere, BlueprintReadWrite, AdvancedDisplay, Category = "Training Settings", meta = (ClampMin = "0.1", ClampMax = "1.0"))
@@ -263,7 +279,7 @@ public:
 
 	/** The activation function to use in the neural network. */
 	UPROPERTY(EditAnywhere, BlueprintReadWrite, AdvancedDisplay, Category = "Training Settings")
-	EActivationFunction ActivationFunction = EActivationFunction::Relu;
+	EActivationFunction ActivationFunction = EActivationFunction::LRelu;
 
 	/** The loss function to use during model training. */
 	UPROPERTY(EditAnywhere, BlueprintReadWrite, AdvancedDisplay, Category = "Training Settings")
@@ -281,6 +297,10 @@ public:
 	UPROPERTY(EditAnywhere, BlueprintReadWrite, AdvancedDisplay, Category = "Training Settings", meta = (ClampMin = "0.0", ClampMax = "100.0")) // A percentage.
 	float NoiseAmount = 0.5;
 
+	/** The delta mode, which describes in what space the calculated deltas are. */
+	UPROPERTY(EditAnywhere, BlueprintReadWrite, AdvancedDisplay, Category = "Training Settings")
+	EDeltaMode DeltaMode = EDeltaMode::PreSkinning;
+
 	/** The loss function to use during model training. */
 	UPROPERTY(EditAnywhere, BlueprintReadWrite, AdvancedDisplay, Category = "Training Settings")
 	EDeviceType DeviceType = EDeviceType::GPU;
diff --git a/Engine/Plugins/Experimental/MLDeformer/Source/MLDeformerEditor/Private/MLDeformerAssetActions.h b/Engine/Plugins/Experimental/MLDeformer/Source/MLDeformerEditor/Private/MLDeformerAssetActions.h
--- a/Engine/Plugins/Experimental/MLDeformer/Source/MLDeformerEditor/Private/MLDeformerAssetActions.h
+++ b/Engine/Plugins/Experimental/MLDeformer/Source/MLDeformerEditor/Private/MLDeformerAssetActions.h
@@ -10,7 +10,7 @@ class FMLDeformerAssetActions : public FAssetTypeActions_Base
 {
 public:
 	// IAssetTypeActions Implementation
-	virtual FText GetName() const override { return NSLOCTEXT("AssetTypeActions", "AssetTypeActions_MLDeformer", "ML Deformer"); }
+	virtual FText GetName() const override { return NSLOCTEXT("MLDeformerAssetActions", "Name", "ML Deformer"); }
 	virtual void GetActions(const TArray<UObject*>& InObjects, struct FToolMenuSection& Section) override;
 	virtual FColor GetTypeColor() const override { return FColor(255, 255, 0); }
 	virtual UClass* GetSupportedClass() const override;
diff --git a/Engine/Plugins/Experimental/MLDeformer/Source/MLDeformerEditor/Private/MLDeformerAssetDetails.cpp b/Engine/Plugins/Experimental/MLDeformer/Source/MLDeformerEditor/Private/MLDeformerAssetDetails.cpp
--- a/Engine/Plugins/Experimental/MLDeformer/Source/MLDeformerEditor/Private/MLDeformerAssetDetails.cpp
+++ b/Engine/Plugins/Experimental/MLDeformer/Source/MLDeformerEditor/Private/MLDeformerAssetDetails.cpp
@@ -243,6 +243,7 @@ void FMLDeformerAssetDetails::CustomizeDetails(class IDetailLayoutBuilder& Detai
 	SettingsCategoryBuilder.AddProperty(GET_MEMBER_NAME_CHECKED(UMLDeformerAsset, BatchSize));
 	SettingsCategoryBuilder.AddProperty(GET_MEMBER_NAME_CHECKED(UMLDeformerAsset, Epochs));
 	SettingsCategoryBuilder.AddProperty(GET_MEMBER_NAME_CHECKED(UMLDeformerAsset, EpochsWithDecay));
+	SettingsCategoryBuilder.AddProperty(GET_MEMBER_NAME_CHECKED(UMLDeformerAsset, MaxTrainingFrames));
 	SettingsCategoryBuilder.AddProperty(GET_MEMBER_NAME_CHECKED(UMLDeformerAsset, CacheSizeInMegabytes));
 
 	// Show a warning when no neural network has been set.
@@ -280,6 +281,7 @@ void FMLDeformerAssetDetails::CustomizeDetails(class IDetailLayoutBuilder& Detai
 		}
 	}
 
+	SettingsCategoryBuilder.AddProperty(GET_MEMBER_NAME_CHECKED(UMLDeformerAsset, DeltaMode));
 	SettingsCategoryBuilder.AddProperty(GET_MEMBER_NAME_CHECKED(UMLDeformerAsset, DecayFunction));
 	SettingsCategoryBuilder.AddProperty(GET_MEMBER_NAME_CHECKED(UMLDeformerAsset, DecayRate));
 	SettingsCategoryBuilder.AddProperty(GET_MEMBER_NAME_CHECKED(UMLDeformerAsset, LearningRate));
diff --git a/Engine/Plugins/Experimental/MLDeformer/Source/MLDeformerEditor/Private/MLDeformerEditor.cpp b/Engine/Plugins/Experimental/MLDeformer/Source/MLDeformerEditor/Private/MLDeformerEditor.cpp
--- a/Engine/Plugins/Experimental/MLDeformer/Source/MLDeformerEditor/Private/MLDeformerEditor.cpp
+++ b/Engine/Plugins/Experimental/MLDeformer/Source/MLDeformerEditor/Private/MLDeformerEditor.cpp
@@ -16,6 +16,30 @@
 
 IMPLEMENT_MODULE(FMLDeformerEditor, MLDeformerEditor)
 
+namespace MLDeformerCVars
+{
+	// Enable or disable debug drawing of the first debug data.
+	TAutoConsoleVariable<bool> DebugDraw1(
+		TEXT("MLDeformer.DebugDraw1"),
+		false,
+		TEXT("Should debug drawing be enabled for the first debug data? Default: false."),
+		ECVF_Default);
+
+	// Enable or disable debug drawing of the second debug data.
+	TAutoConsoleVariable<bool> DebugDraw2(
+		TEXT("MLDeformer.DebugDraw2"),
+		false,
+		TEXT("Should debug drawing be enabled for the second debug data? Default: false."),
+		ECVF_Default);
+
+	// The debug data point size.
+	TAutoConsoleVariable<float> DebugDrawPointSize(
+		TEXT("MLDeformer.DebugDrawPointSize"),
+		1.5f,
+		TEXT("The size of the points when debug drawing is enabled Default: 1.5."),
+		ECVF_Default);
+}
+
 void FMLDeformerEditor::StartupModule()
 {
 	MLDeformerAssetActions = MakeShareable(new FMLDeformerAssetActions);
diff --git a/Engine/Plugins/Experimental/MLDeformer/Source/MLDeformerEditor/Private/MLDeformerEditorData.cpp b/Engine/Plugins/Experimental/MLDeformer/Source/MLDeformerEditor/Private/MLDeformerEditorData.cpp
--- a/Engine/Plugins/Experimental/MLDeformer/Source/MLDeformerEditor/Private/MLDeformerEditorData.cpp
+++ b/Engine/Plugins/Experimental/MLDeformer/Source/MLDeformerEditor/Private/MLDeformerEditorData.cpp
@@ -35,6 +35,7 @@
 
 #define LOCTEXT_NAMESPACE "MLDeformerEditorData"
 
+
 FMLDeformerEditorData::FMLDeformerEditorData()
 {
 	Actors.AddDefaulted(5);
@@ -356,9 +357,23 @@ bool FMLDeformerEditorData::IsReadyForTraining() const
 	}
 
 	// Make sure every skeletal imported mesh has some geometry track.
-	if (!GetSingleFrameCache().GetSampler().GetFailedImportedMeshNames().IsEmpty())
+	const int32 NumGeomCacheTracks = GetDeformerAsset()->GetGeometryCache() ?  GetDeformerAsset()->GetGeometryCache()->Tracks.Num() : 0;
+	int32 NumSkelMeshes = 0;
+	if (GetDeformerAsset()->GetSkeletalMesh())
 	{
-		return false;
+		FSkeletalMeshModel* Model = GetDeformerAsset()->GetSkeletalMesh()->GetImportedModel();
+		if (Model)
+		{
+			NumSkelMeshes = Model->LODModels[0].ImportedMeshInfos.Num();		
+		}
+	}
+
+	if (NumGeomCacheTracks != 1 || NumSkelMeshes != 1)	// Allow the special case where there is just one mesh and track.
+	{
+		if (!GetSingleFrameCache().GetSampler().GetFailedImportedMeshNames().IsEmpty())
+		{
+			return false;
+		}
 	}
 
 	return true;
@@ -379,6 +394,8 @@ bool FMLDeformerEditorData::GenerateDeltas(uint32 LODIndex, uint32 FrameNumber,
 	// Let's also directly extract the positions of the vertices, since we already calculated them.
 	const FMLDeformerSamplerData& SamplerData = SingleFrameCache.GetSampler().GetSamplerData();
 	LinearSkinnedPositions = SamplerData.GetSkinnedVertexPositions();
+	DebugVectors = SamplerData.GetDebugVectors();
+	DebugVectors2 = SamplerData.GetDebugVectors2();
 
 	return true;
 }
@@ -427,7 +444,7 @@ bool FMLDeformerEditorData::ComputeVertexDeltaStatistics(uint32 LODIndex, FMLDef
 
 	// Calculate the number of frames to output.
 	UGeometryCacheComponent* GeometryCacheComponent = GetEditorActor(EMLDeformerEditorActorIndex::Target).GeomCacheComponent;
-	const uint32 AnimNumFrames = static_cast<uint32>(GeometryCacheComponent->GetNumberOfFrames());
+	const uint32 AnimNumFrames = MLDeformerAsset->GetNumFramesForTraining();
 
 	// Show some dialog with progress.
 	const FText Title(LOCTEXT("PreprocessTrainingDataMessage", "Calculating data statistics"));
@@ -439,7 +456,8 @@ bool FMLDeformerEditorData::ComputeVertexDeltaStatistics(uint32 LODIndex, FMLDef
 	FVector3f VertexDeltaScale = FVector3f::OneVector;
 	float Count = 0.0f;
 	USkeletalMeshComponent* SkeletalMeshComponent = GetEditorActor(EMLDeformerEditorActorIndex::Base).SkelMeshComponent;
-	for (uint32 FrameIndex = 0; FrameIndex < AnimNumFrames; FrameIndex++)
+	const int32 FrameCount = MLDeformerAsset->GetNumFramesForTraining();
+	for (int32 FrameIndex = 0; FrameIndex < FrameCount; FrameIndex++)
 	{
 		// Get the deltas updated for this frame, and then update the mean and scale.
 		const FMLDeformerTrainingFrame& TrainingFrame = InFrameCache->GetTrainingFrameForAnimFrame(FrameIndex);
@@ -684,7 +702,7 @@ bool FMLDeformerEditorData::IsPlayingAnim() const
 
 FString FMLDeformerEditorData::GetDefaultDeformerGraphAssetPath()
 {
-	return FString(TEXT("/MLDeformer/Deformers/DefaultMLDeformerGraph.MLDeformerGraph"));
+	return FString(TEXT("/MLDeformer/Deformers/PreSkinningMLDeformerGraph.PreSkinningMLDeformerGraph"));
 }
 
 UComputeGraph* FMLDeformerEditorData::LoadDefaultDeformerGraph()
diff --git a/Engine/Plugins/Experimental/MLDeformer/Source/MLDeformerEditor/Private/MLDeformerEditorMode.cpp b/Engine/Plugins/Experimental/MLDeformer/Source/MLDeformerEditor/Private/MLDeformerEditorMode.cpp
--- a/Engine/Plugins/Experimental/MLDeformer/Source/MLDeformerEditor/Private/MLDeformerEditorMode.cpp
+++ b/Engine/Plugins/Experimental/MLDeformer/Source/MLDeformerEditor/Private/MLDeformerEditorMode.cpp
@@ -1,6 +1,7 @@
 // Copyright Epic Games, Inc. All Rights Reserved.
 
 #include "MLDeformerEditorMode.h"
+#include "MLDeformerEditor.h"
 #include "MLDeformerEditorData.h"
 #include "MLDeformerAssetDetails.h"
 #include "MLDeformerComponent.h"
@@ -79,6 +80,16 @@ void FMLDeformerEditorMode::GetOnScreenDebugInfo(TArray<FText>& OutDebugInfo) co
 {
 }
 
+void FMLDeformerEditorMode::DrawDebugPoints(FPrimitiveDrawInterface* PDI, const TArray<FVector3f>& Points, int32 DepthGroup, const FLinearColor& Color)
+{
+	const float PointSize = MLDeformerCVars::DebugDrawPointSize.GetValueOnAnyThread();
+	for (int32 Index = 0; Index < Points.Num(); ++Index)
+	{
+		const FVector Position = Points[Index];
+		PDI->DrawPoint(Position, Color, PointSize, DepthGroup);
+	}
+}
+
 void FMLDeformerEditorMode::Render(const FSceneView* View, FViewport* Viewport, FPrimitiveDrawInterface* PDI)
 {
 	FEdMode::Render(View, Viewport, PDI);
@@ -97,20 +108,37 @@ void FMLDeformerEditorMode::Render(const FSceneView* View, FViewport* Viewport,
 		if (VizSettings->GetDrawVertexDeltas() && (Data->VertexDeltas.Num() / 3) == Data->LinearSkinnedPositions.Num())
 		{
 			const FLinearColor DeltasColor = FMLDeformerEditorStyle::Get().GetColor("MLDeformer.Deltas.Color");
+			const FLinearColor DebugVectorsColor = FMLDeformerEditorStyle::Get().GetColor("MLDeformer.DebugVectors.Color");
+			const FLinearColor DebugVectorsColor2 = FMLDeformerEditorStyle::Get().GetColor("MLDeformer.DebugVectors.Color2");
 			const uint8 DepthGroup = VizSettings->GetXRayDeltas() ? 100 : 0;
 			const TArray<FVector3f>& SkinnedPositions = Data->LinearSkinnedPositions;
 			for (int32 Index = 0; Index < Data->LinearSkinnedPositions.Num(); ++Index)
 			{
-				const FVector Start = SkinnedPositions[Index];
-				FVector End;
-				int32 ArrayIndex = 3 * Index;
-				End.X = Start.X + Data->VertexDeltas[ArrayIndex];
-				End.Y = Start.Y + Data->VertexDeltas[++ArrayIndex];
-				End.Z = Start.Z + Data->VertexDeltas[++ArrayIndex];
-
-				PDI->DrawLine(Start, End, DeltasColor, DepthGroup);
+				const int32 ArrayIndex = 3 * Index;
+				const FVector Delta(
+					Data->VertexDeltas[ArrayIndex], 
+					Data->VertexDeltas[ArrayIndex + 1], 
+					Data->VertexDeltas[ArrayIndex + 2]);
+				const FVector VertexPos = SkinnedPositions[Index];
+				PDI->DrawLine(VertexPos, VertexPos + Delta, DeltasColor, DepthGroup);
 			}
 		}
+
+		// Draw the first set of debug points.
+		if (MLDeformerCVars::DebugDraw1.GetValueOnAnyThread())
+		{
+			const uint8 DepthGroup = VizSettings->GetXRayDeltas() ? 100 : 0;
+			const FLinearColor Color = FMLDeformerEditorStyle::Get().GetColor("MLDeformer.DebugVectors.Color");
+			DrawDebugPoints(PDI, Data->DebugVectors, DepthGroup, Color);
+		}
+
+		// Draw the second set of debug points.
+		if (MLDeformerCVars::DebugDraw2.GetValueOnAnyThread())
+		{
+			const uint8 DepthGroup = VizSettings->GetXRayDeltas() ? 100 : 0;
+			const FLinearColor Color = FMLDeformerEditorStyle::Get().GetColor("MLDeformer.DebugVectors.Color2");
+			DrawDebugPoints(PDI, Data->DebugVectors2, DepthGroup, Color);
+		}
 	}
 }
 
diff --git a/Engine/Plugins/Experimental/MLDeformer/Source/MLDeformerEditor/Private/MLDeformerEditorStyle.cpp b/Engine/Plugins/Experimental/MLDeformer/Source/MLDeformerEditor/Private/MLDeformerEditorStyle.cpp
--- a/Engine/Plugins/Experimental/MLDeformer/Source/MLDeformerEditor/Private/MLDeformerEditorStyle.cpp
+++ b/Engine/Plugins/Experimental/MLDeformer/Source/MLDeformerEditor/Private/MLDeformerEditorStyle.cpp
@@ -40,6 +40,9 @@ FMLDeformerEditorStyle::FMLDeformerEditorStyle() : FSlateStyleSet("MLDeformerEdi
 	Set("MLDeformer.Vertex.Color", FLinearColor(1.0f, 0.0f, 1.0f));
 	Set("MLDeformer.Vertex.Size", 5.0f);
 
+	Set("MLDeformer.DebugVectors.Color", FLinearColor(1.0f, 0.0f, 0.0f));
+	Set("MLDeformer.DebugVectors.Color2", FLinearColor(1.0f, 1.0f, 0.0f));
+
 	Set("MLDeformer.DefaultLabelScale", 0.65f);
 
 	FSlateStyleRegistry::RegisterSlateStyle(*this);
diff --git a/Engine/Plugins/Experimental/MLDeformer/Source/MLDeformerEditor/Private/MLDeformerEditorToolkit.cpp b/Engine/Plugins/Experimental/MLDeformer/Source/MLDeformerEditor/Private/MLDeformerEditorToolkit.cpp
--- a/Engine/Plugins/Experimental/MLDeformer/Source/MLDeformerEditor/Private/MLDeformerEditorToolkit.cpp
+++ b/Engine/Plugins/Experimental/MLDeformer/Source/MLDeformerEditor/Private/MLDeformerEditorToolkit.cpp
@@ -183,7 +183,6 @@ void FMLDeformerEditorToolkit::FillToolbar(FToolBarBuilder& ToolbarBuilder)
 						check(!DeformerAsset->GetInputInfo().IsEmpty());
 
 						// Init the frame cache.
-						const int32 NumFrames = (DeformerAsset->GetGeometryCache()->GetEndFrame() - DeformerAsset->GetGeometryCache()->GetStartFrame()) + 1;
 						FMLDeformerFrameCache::FInitSettings FrameCacheInitSettings;
 						FrameCacheInitSettings.DeformerAsset = DeformerAsset;
 						FrameCacheInitSettings.CacheSizeInBytes = 1024ull * 1024 * DeformerAsset->GetCacheSizeInMegabytes();
@@ -715,7 +714,9 @@ void FMLDeformerEditorToolkit::OnFinishedChangingDetails(const FPropertyChangedE
 	}
 	else
 	if (Property->GetFName() == GET_MEMBER_NAME_CHECKED(UMLDeformerAsset, DeltaCutoffLength) ||
-	    Property->GetFName() == GET_MEMBER_NAME_CHECKED(UMLDeformerAsset, AlignmentTransform))
+	    Property->GetFName() == GET_MEMBER_NAME_CHECKED(UMLDeformerAsset, AlignmentTransform) ||
+	    Property->GetFName() == GET_MEMBER_NAME_CHECKED(UMLDeformerAsset, DeltaMode) ||
+	    Property->GetFName() == GET_MEMBER_NAME_CHECKED(UMLDeformerAsset, MaxTrainingFrames))
 	{
 		EditorData->InitAssets();
 	}
diff --git a/Engine/Plugins/Experimental/MLDeformer/Source/MLDeformerEditor/Private/MLDeformerFrameCache.cpp b/Engine/Plugins/Experimental/MLDeformer/Source/MLDeformerEditor/Private/MLDeformerFrameCache.cpp
--- a/Engine/Plugins/Experimental/MLDeformer/Source/MLDeformerEditor/Private/MLDeformerFrameCache.cpp
+++ b/Engine/Plugins/Experimental/MLDeformer/Source/MLDeformerEditor/Private/MLDeformerFrameCache.cpp
@@ -1,12 +1,14 @@
 // Copyright Epic Games, Inc. All Rights Reserved.
 
 #include "MLDeformerFrameCache.h"
+#include "MLDeformerEditor.h"
 
 #include "Animation/DebugSkelMeshComponent.h"
 #include "Animation/AnimSequence.h"
 #include "Engine/SkeletalMesh.h"
 #include "Rendering/SkeletalMeshModel.h"
 #include "Rendering/SkeletalMeshLODModel.h"
+#include "Rendering/SkeletalMeshLODRenderData.h"
 
 #include "GeometryCacheComponent.h"
 #include "GeometryCache.h"
@@ -15,6 +17,7 @@
 
 #include "Math/NumericLimits.h"
 
+
 void FMLDeformerTrainingFrame::Clear()
 {
 	VertexDeltas.Empty();
@@ -80,17 +83,20 @@ void FMLDeformerSamplerData::Update(int32 InAnimFrameIndex)
 	const float SampleTime = GeometryCacheComponent->GetTimeAtFrame(AnimFrameIndex);
 	if (SkeletalMeshComponent && SkeletalMeshComponent->SkeletalMesh)
 	{
-		// Sample the transforms at the frame time and extract the bone rotations.
+		// Sample the transforms at the frame time.
 		SkeletalMeshComponent->SetPosition(SampleTime);
 		SkeletalMeshComponent->bPauseAnims = true;
 		SkeletalMeshComponent->RefreshBoneTransforms();
+		SkeletalMeshComponent->CacheRefToLocalMatrices(BoneMatrices);
+
+		// Extract bone rotations to feed to the neural net.
 		const FMLDeformerInputInfo& InputInfo = DeformerAsset.GetInputInfo();
 		InputInfo.ExtractBoneRotations(SkeletalMeshComponent, BoneRotations);
 
-		// Get the bone transform matrices used during skinning, then extract the skinned positions.
-		SkeletalMeshComponent->CacheRefToLocalMatrices(BoneMatrices);
+		// Calculate skinned positions.
 		ExtractSkinnedPositions(LODIndex, BoneMatrices, TempVertexPositions, SkinnedVertexPositions);
 
+		// Calculate curve values.
 		InputInfo.ExtractCurveValues(SkeletalMeshComponent, CurveValues);
 	}
 
@@ -104,15 +110,31 @@ void FMLDeformerSamplerData::Update(int32 InAnimFrameIndex)
 
 	// Calculate the vertex deltas.
 	const float DeltaCutoffLength = DeformerAsset.GetDeltaCutoffLength();
-	CalculateVertexDeltas(SkinnedVertexPositions, DeltaCutoffLength, VertexDeltas);
+	CalculateVertexDeltas(SkinnedVertexPositions, DeltaCutoffLength, VertexDeltas, DebugVectors, DebugVectors2);
 }
 
-void FMLDeformerSamplerData::CalculateVertexDeltas(const TArray<FVector3f>& SkinnedPositions, float DeltaCutoffLength, TArray<float>& OutVertexDeltas) const
+void FMLDeformerSamplerData::CalculateVertexDeltas(const TArray<FVector3f>& SkinnedPositions, float DeltaCutoffLength, TArray<float>& OutVertexDeltas, TArray<FVector3f>& OutDebugVectors, TArray<FVector3f>& OutDebugVectors2) const
 {
+	const int32 LODIndex = 0;
 	const int32 NumVerts = SkinnedPositions.Num();
 	OutVertexDeltas.Reset(NumVerts * 3);
 	OutVertexDeltas.AddUninitialized(NumVerts * 3);
 
+	// Debug vectors.
+	OutDebugVectors.Reset();
+	OutDebugVectors2.Reset();
+	const bool bDebugDraw = (MLDeformerCVars::DebugDraw1.GetValueOnAnyThread() || MLDeformerCVars::DebugDraw2.GetValueOnAnyThread());
+	if (bDebugDraw)
+	{
+		OutDebugVectors.AddUninitialized(NumVerts);
+		OutDebugVectors2.AddUninitialized(NumVerts);
+		for (int32 Index = 0; Index < NumVerts; ++Index)
+		{
+			OutDebugVectors[Index] = FVector3f::ZeroVector;
+			OutDebugVectors2[Index] = FVector3f::ZeroVector;
+		}
+	}
+
 	UMLDeformerAsset* DeformerAsset = Sampler->GetInitSettings().DeformerAsset;
 	USkeletalMesh* SkelMesh = DeformerAsset->GetSkeletalMesh();
 	UGeometryCache* GeomCache = DeformerAsset->GetGeometryCache();
@@ -132,7 +154,8 @@ void FMLDeformerSamplerData::CalculateVertexDeltas(const TArray<FVector3f>& Skin
 
 	// For all mesh mappings we found.
 	const float SampleTime = GeometryCacheComponent->GetTimeAtFrame(AnimFrameIndex);
-	const TArray<FSkelMeshImportedMeshInfo>& SkelMeshInfos = ImportedModel->LODModels[0].ImportedMeshInfos;
+	const FSkeletalMeshLODModel& LODModel = ImportedModel->LODModels[LODIndex];
+	const TArray<FSkelMeshImportedMeshInfo>& SkelMeshInfos = LODModel.ImportedMeshInfos;
 	for (int32 MeshMappingIndex = 0; MeshMappingIndex < Sampler->GetNumMeshMappings(); ++MeshMappingIndex)
 	{
 		const FMLDeformerMeshMapping& MeshMapping = Sampler->GetMeshMapping(MeshMappingIndex);
@@ -147,18 +170,47 @@ void FMLDeformerSamplerData::CalculateVertexDeltas(const TArray<FVector3f>& Skin
 		}
 
 		// Calculate the vertex deltas.
+		const FSkeletalMeshLODRenderData& SkelMeshLODData = SkelMesh->GetResourceForRendering()->LODRenderData[LODIndex];
+		const FSkinWeightVertexBuffer& SkinWeightBuffer = *SkeletalMeshComponent->GetSkinWeightBuffer(LODIndex);
 		for (int32 VertexIndex = 0; VertexIndex < MeshInfo.NumVertices; ++VertexIndex)
 		{
 			const int32 SkinnedVertexIndex = MeshInfo.StartImportedVertex + VertexIndex;
 			const int32 GeomCacheVertexIndex = MeshMapping.SkelMeshToTrackVertexMap[VertexIndex];
-			if (GeomCacheVertexIndex != INDEX_NONE)
+			if (GeomCacheVertexIndex != INDEX_NONE && GeomCacheMeshData.Positions.IsValidIndex(GeomCacheVertexIndex))
 			{
-				const FVector3f SkinnedVertexPos = SkinnedPositions[SkinnedVertexIndex];
-				const FVector3f GeomCacheVertexPos = AlignmentTransform.TransformPosition(GeomCacheMeshData.Positions[GeomCacheVertexIndex]);
-				const FVector3f Delta = GeomCacheVertexPos - SkinnedVertexPos;
+				FVector3f Delta = FVector3f::ZeroVector;
+
+				const int32 ArrayIndex = 3 * SkinnedVertexIndex;
+				if (DeformerAsset->GetDeltaMode() == EDeltaMode::PreSkinning)
+				{
+					// Calculate the inverse skinning transform for this vertex.
+					const int32 RenderVertexIndex = MeshMapping.ImportedVertexToRenderVertexMap[VertexIndex];
+					const FMatrix44f InvSkinningTransform = CalcInverseSkinningTransform(RenderVertexIndex, SkelMeshLODData, SkinWeightBuffer);
+
+					// Calculate the pre-skinning data.
+					const FSkeletalMeshLODRenderData& LODData = SkelMesh->GetResourceForRendering()->LODRenderData[0];
+					const FVector3f UnskinnedPosition = LODData.StaticVertexBuffers.PositionVertexBuffer.VertexPosition(RenderVertexIndex);
+					const FVector3f GeomCacheVertexPos = AlignmentTransform.TransformPosition(GeomCacheMeshData.Positions[GeomCacheVertexIndex]);
+					const FVector3f PreSkinningTargetPos = InvSkinningTransform.TransformPosition(GeomCacheVertexPos);
+					Delta = PreSkinningTargetPos - UnskinnedPosition;
+
+					if (bDebugDraw)
+					{
+						OutDebugVectors[SkinnedVertexIndex] = UnskinnedPosition;
+						OutDebugVectors2[SkinnedVertexIndex] = UnskinnedPosition + Delta;
+					}
+				}
+				else // We're post skinning.
+				{
+					check(DeformerAsset->GetDeltaMode() == EDeltaMode::PostSkinning);
+					const FVector3f SkinnedVertexPos = SkinnedPositions[SkinnedVertexIndex];
+					const FVector3f GeomCacheVertexPos = AlignmentTransform.TransformPosition(GeomCacheMeshData.Positions[GeomCacheVertexIndex]);
+					Delta = GeomCacheVertexPos - SkinnedVertexPos;
+				}
+
+				// Set the delta.
 				if (Delta.Length() < DeltaCutoffLength)
 				{
-					const int32 ArrayIndex = 3 * SkinnedVertexIndex;
 					OutVertexDeltas[ArrayIndex] = Delta.X;
 					OutVertexDeltas[ArrayIndex + 1] = Delta.Y;
 					OutVertexDeltas[ArrayIndex + 2] = Delta.Z;
@@ -168,6 +220,42 @@ void FMLDeformerSamplerData::CalculateVertexDeltas(const TArray<FVector3f>& Skin
 	}
 }
 
+// Calculate the inverse skinning transform. This is basically inv(sum(BoneTransform_i * inv(BoneRestTransform_i) * Weight_i)), where i is for each skinning influence for the given vertex.
+FMatrix44f FMLDeformerSamplerData::CalcInverseSkinningTransform(int32 VertexIndex, const FSkeletalMeshLODRenderData& SkelMeshLODData, const FSkinWeightVertexBuffer& SkinWeightBuffer) const
+{
+	check(SkeletalMeshComponent);
+	const USkeletalMesh* Mesh = SkeletalMeshComponent->SkeletalMesh;
+	check(Mesh);
+
+	// Find the render section, which we need to find the right bone index.
+	int32 SectionIndex = INDEX_NONE;
+	int32 SectionVertexIndex = INDEX_NONE;
+	const FSkeletalMeshLODRenderData& LODData = Mesh->GetResourceForRendering()->LODRenderData[0];
+	LODData.GetSectionFromVertexIndex(VertexIndex, SectionIndex, SectionVertexIndex);
+
+	// Init the matrix at full zeros.
+	FMatrix44f InvSkinningTransform = FMatrix(FVector::ZeroVector, FVector::ZeroVector, FVector::ZeroVector, FVector::ZeroVector);
+	InvSkinningTransform.M[3][3] = 0.0f;
+
+	// For each influence, sum up the weighted skinning matrices.
+	const int32 NumInfluences = SkinWeightBuffer.GetMaxBoneInfluences();
+	for (int32 InfluenceIndex = 0; InfluenceIndex < NumInfluences; ++InfluenceIndex)
+	{
+		const int32 BoneIndex = SkinWeightBuffer.GetBoneIndex(VertexIndex, InfluenceIndex);
+		const uint8 WeightByte = SkinWeightBuffer.GetBoneWeight(VertexIndex, InfluenceIndex);
+		if (WeightByte > 0)
+		{
+			const int32 RealBoneIndex = LODData.RenderSections[SectionIndex].BoneMap[BoneIndex];
+			const float	Weight = static_cast<float>(WeightByte) / 255.0f;
+			const FMatrix44f& SkinningTransform = BoneMatrices[RealBoneIndex];
+			InvSkinningTransform += SkinningTransform * Weight;
+		}
+	}
+
+	// Return the inverse skinning transform matrix.
+	return InvSkinningTransform.Inverse();
+}
+
 void FMLDeformerSamplerData::ExtractSkinnedPositions(int32 LODIndex, TArray<FMatrix44f>& InBoneMatrices, TArray<FVector3f>& TempPositions, TArray<FVector3f>& OutPositions) const
 {
 	OutPositions.Reset();
@@ -215,6 +303,8 @@ SIZE_T FMLDeformerSamplerData::CalcMemUsageInBytes() const
 	SIZE_T NumBytes = 0;
 	NumBytes += SkinnedVertexPositions.GetAllocatedSize();
 	NumBytes += TempVertexPositions.GetAllocatedSize();
+	NumBytes += DebugVectors.GetAllocatedSize();
+	NumBytes += DebugVectors2.GetAllocatedSize();
 	NumBytes += BoneMatrices.GetAllocatedSize();
 	NumBytes += VertexDeltas.GetAllocatedSize();
 	NumBytes += BoneRotations.GetAllocatedSize();
@@ -335,7 +425,7 @@ int32 FMLDeformerSampler::GetNumCurves() const
 
 int32 FMLDeformerSampler::GetNumFrames() const
 {
-	return InitSettings.DeformerAsset->GetNumFrames();
+	return InitSettings.DeformerAsset->GetNumFramesForTraining();
 }
 
 SIZE_T FMLDeformerSampler::CalcMemUsageInBytes() const
diff --git a/Engine/Plugins/Experimental/MLDeformer/Source/MLDeformerEditor/Private/MLDeformerFrameCache.h b/Engine/Plugins/Experimental/MLDeformer/Source/MLDeformerEditor/Private/MLDeformerFrameCache.h
--- a/Engine/Plugins/Experimental/MLDeformer/Source/MLDeformerEditor/Private/MLDeformerFrameCache.h
+++ b/Engine/Plugins/Experimental/MLDeformer/Source/MLDeformerEditor/Private/MLDeformerFrameCache.h
@@ -20,6 +20,8 @@ class UWorld;
 class AActor;
 class FMLDeformerSampler;
 class FMLDeformerFrameCache;
+class FSkeletalMeshLODRenderData;
+class FSkinWeightVertexBuffer;
 
 /**
  * A sampler data object, which is basically a set of data that is used to generate training data for a given frame.
@@ -44,6 +46,8 @@ public:
 	UGeometryCacheComponent* GetGeometryCacheComponent() const { return GeometryCacheComponent; }
 
 	const TArray<FVector3f>& GetSkinnedVertexPositions() const { return SkinnedVertexPositions; }
+	const TArray<FVector3f>& GetDebugVectors() const { return DebugVectors; }
+	const TArray<FVector3f>& GetDebugVectors2() const { return DebugVectors2; }
 
 	const TArray<float>& GetVertexDeltas() const { return VertexDeltas; }
 	const TArray<float>& GetBoneRotations() const { return BoneRotations; }
@@ -56,7 +60,8 @@ public:
 
 protected:
 	void ExtractSkinnedPositions(int32 LODIndex, TArray<FMatrix44f>& InBoneMatrices, TArray<FVector3f>& TempPositions, TArray<FVector3f>& OutPositions) const;
-	void CalculateVertexDeltas(const TArray<FVector3f>& SkinnedPositions, float DeltaCutoffLength, TArray<float>& OutVertexDeltas) const;
+	void CalculateVertexDeltas(const TArray<FVector3f>& SkinnedPositions, float DeltaCutoffLength, TArray<float>& OutVertexDeltas, TArray<FVector3f>& OutDebugVectors, TArray<FVector3f>& OutDebugVectors2) const;
+	FMatrix44f CalcInverseSkinningTransform(int32 VertexIndex, const FSkeletalMeshLODRenderData& SkelMeshLODData, const FSkinWeightVertexBuffer& SkinWeightBuffer) const;
 
 protected:
 	FMLDeformerSampler* Sampler = nullptr;
@@ -64,6 +69,8 @@ protected:
 	TObjectPtr<UGeometryCacheComponent> GeometryCacheComponent = nullptr;
 	TArray<FVector3f> SkinnedVertexPositions;
 	TArray<FVector3f> TempVertexPositions;
+	TArray<FVector3f> DebugVectors;
+	TArray<FVector3f> DebugVectors2;
 	TArray<FMatrix44f> BoneMatrices;
 	TArray<float> VertexDeltas;	// (NumImportedVerts * 3) -> xyz
 	TArray<float> BoneRotations; // (NumBones * 4) -> quat xyzw
diff --git a/Engine/Plugins/Experimental/MLDeformer/Source/MLDeformerEditor/Private/MLPytorchDataSetInterface.cpp b/Engine/Plugins/Experimental/MLDeformer/Source/MLDeformerEditor/Private/MLPytorchDataSetInterface.cpp
--- a/Engine/Plugins/Experimental/MLDeformer/Source/MLDeformerEditor/Private/MLPytorchDataSetInterface.cpp
+++ b/Engine/Plugins/Experimental/MLDeformer/Source/MLDeformerEditor/Private/MLPytorchDataSetInterface.cpp
@@ -65,7 +65,7 @@ int32 UMLPytorchDataSetInterface::GetNumberSampleDeltas() const
 int32 UMLPytorchDataSetInterface::NumSamples() const
 {
 	check(IsValid());
-	return EditorData->GetDeformerAsset()->GetNumFrames();
+	return EditorData->GetDeformerAsset()->GetNumFramesForTraining();
 }
 
 bool UMLPytorchDataSetInterface::SetCurrentSampleIndex(int32 Index)
diff --git a/Engine/Plugins/Experimental/MLDeformer/Source/MLDeformerEditor/Public/MLDeformerEditor.h b/Engine/Plugins/Experimental/MLDeformer/Source/MLDeformerEditor/Public/MLDeformerEditor.h
--- a/Engine/Plugins/Experimental/MLDeformer/Source/MLDeformerEditor/Public/MLDeformerEditor.h
+++ b/Engine/Plugins/Experimental/MLDeformer/Source/MLDeformerEditor/Public/MLDeformerEditor.h
@@ -4,9 +4,17 @@
 
 #include "CoreMinimal.h"
 #include "Modules/ModuleInterface.h"
+#include "HAL/IConsoleManager.h"
 
 class FMLDeformerAssetActions;
 
+namespace MLDeformerCVars
+{
+	extern TAutoConsoleVariable<bool> DebugDraw1;
+	extern TAutoConsoleVariable<bool> DebugDraw2;
+	extern TAutoConsoleVariable<float> DebugDrawPointSize;
+};
+
 class MLDEFORMEREDITOR_API FMLDeformerEditor : public IModuleInterface
 {
 public:
diff --git a/Engine/Plugins/Experimental/MLDeformer/Source/MLDeformerEditor/Public/MLDeformerEditorData.h b/Engine/Plugins/Experimental/MLDeformer/Source/MLDeformerEditor/Public/MLDeformerEditorData.h
--- a/Engine/Plugins/Experimental/MLDeformer/Source/MLDeformerEditor/Public/MLDeformerEditorData.h
+++ b/Engine/Plugins/Experimental/MLDeformer/Source/MLDeformerEditor/Public/MLDeformerEditorData.h
@@ -128,7 +128,10 @@ public:
 	
 public:
 	TArray<FVector3f> LinearSkinnedPositions;
+	TArray<FVector3f> DebugVectors;
+	TArray<FVector3f> DebugVectors2;
 	TArray<float> VertexDeltas;
+
 	bool bIsVertexDeltaNormalized = false;
 	int32 CurrentFrame = -1;
 
diff --git a/Engine/Plugins/Experimental/MLDeformer/Source/MLDeformerEditor/Public/MLDeformerEditorMode.h b/Engine/Plugins/Experimental/MLDeformer/Source/MLDeformerEditor/Public/MLDeformerEditorMode.h
--- a/Engine/Plugins/Experimental/MLDeformer/Source/MLDeformerEditor/Public/MLDeformerEditorMode.h
+++ b/Engine/Plugins/Experimental/MLDeformer/Source/MLDeformerEditor/Public/MLDeformerEditorMode.h
@@ -42,6 +42,7 @@ private:
 	void UpdateActors();
 	void UpdateLabels();
 	void EncapsulateBounds(const FMLDeformerEditorActor* Actor, FBox& Box) const;
+	void DrawDebugPoints(FPrimitiveDrawInterface* PDI, const TArray<FVector3f>& Points, int32 DepthGroup, const FLinearColor& Color);
 
 private:
 	TWeakPtr<FMLDeformerEditorData> EditorData;
