MLDeformer: Early out of DataInterface work if NN is not loaded yet.

diff --git a/Engine/Plugins/Experimental/MLDeformer/Source/MLDeformer/Private/MLDeformerGraphDataInterface.cpp b/Engine/Plugins/Experimental/MLDeformer/Source/MLDeformer/Private/MLDeformerGraphDataInterface.cpp
--- a/Engine/Plugins/Experimental/MLDeformer/Source/MLDeformer/Private/MLDeformerGraphDataInterface.cpp
+++ b/Engine/Plugins/Experimental/MLDeformer/Source/MLDeformer/Private/MLDeformerGraphDataInterface.cpp
@@ -89,7 +89,9 @@ bool UMLDeformerDataProvider::IsValid() const
 	return
 		DeformerComponent != nullptr &&
 		DeformerComponent->GetDeformerAsset() != nullptr &&
-		DeformerComponent->GetDeformerAsset()->GetVertexMapBuffer().ShaderResourceViewRHI != nullptr;
+		DeformerComponent->GetDeformerAsset()->GetVertexMapBuffer().ShaderResourceViewRHI != nullptr &&
+		DeformerComponent->GetDeformerAsset()->GetInferenceNeuralNetwork() != nullptr &&
+		DeformerComponent->GetDeformerAsset()->GetInferenceNeuralNetwork()->IsLoaded();
 }
 
 FComputeDataProviderRenderProxy* UMLDeformerDataProvider::GetRenderProxy()
