MLDeformer: Make plugin depend on PythonFoundationPackages

diff --git a/Engine/Plugins/Experimental/MLDeformer/MLDeformer.uplugin b/Engine/Plugins/Experimental/MLDeformer/MLDeformer.uplugin
--- a/Engine/Plugins/Experimental/MLDeformer/MLDeformer.uplugin
+++ b/Engine/Plugins/Experimental/MLDeformer/MLDeformer.uplugin
@@ -38,6 +38,10 @@
 		{
 			"Name": "Optimus",
 			"Enabled": true
+		},
+		{
+			"Name": "PythonFoundationPackages",
+			"Enabled": true
 		}
 	]
 }
\ No newline at end of file
