[MLDeformer] Crash fix.

diff --git a/Engine/Plugins/Experimental/MLDeformer/Source/MLDeformerEditor/Private/MLDeformerFrameCache.cpp b/Engine/Plugins/Experimental/MLDeformer/Source/MLDeformerEditor/Private/MLDeformerFrameCache.cpp
--- a/Engine/Plugins/Experimental/MLDeformer/Source/MLDeformerEditor/Private/MLDeformerFrameCache.cpp
+++ b/Engine/Plugins/Experimental/MLDeformer/Source/MLDeformerEditor/Private/MLDeformerFrameCache.cpp
@@ -185,19 +185,22 @@ void FMLDeformerSamplerData::CalculateVertexDeltas(const TArray<FVector3f>& Skin
 				{
 					// Calculate the inverse skinning transform for this vertex.
 					const int32 RenderVertexIndex = MeshMapping.ImportedVertexToRenderVertexMap[VertexIndex];
-					const FMatrix44f InvSkinningTransform = CalcInverseSkinningTransform(RenderVertexIndex, SkelMeshLODData, SkinWeightBuffer);
-
-					// Calculate the pre-skinning data.
-					const FSkeletalMeshLODRenderData& LODData = SkelMesh->GetResourceForRendering()->LODRenderData[0];
-					const FVector3f UnskinnedPosition = LODData.StaticVertexBuffers.PositionVertexBuffer.VertexPosition(RenderVertexIndex);
-					const FVector3f GeomCacheVertexPos = AlignmentTransform.TransformPosition(GeomCacheMeshData.Positions[GeomCacheVertexIndex]);
-					const FVector3f PreSkinningTargetPos = InvSkinningTransform.TransformPosition(GeomCacheVertexPos);
-					Delta = PreSkinningTargetPos - UnskinnedPosition;
-
-					if (bDebugDraw)
+					if (RenderVertexIndex != INDEX_NONE)
 					{
-						OutDebugVectors[SkinnedVertexIndex] = UnskinnedPosition;
-						OutDebugVectors2[SkinnedVertexIndex] = UnskinnedPosition + Delta;
+						const FMatrix44f InvSkinningTransform = CalcInverseSkinningTransform(RenderVertexIndex, SkelMeshLODData, SkinWeightBuffer);
+
+						// Calculate the pre-skinning data.
+						const FSkeletalMeshLODRenderData& LODData = SkelMesh->GetResourceForRendering()->LODRenderData[0];
+						const FVector3f UnskinnedPosition = LODData.StaticVertexBuffers.PositionVertexBuffer.VertexPosition(RenderVertexIndex);
+						const FVector3f GeomCacheVertexPos = AlignmentTransform.TransformPosition(GeomCacheMeshData.Positions[GeomCacheVertexIndex]);
+						const FVector3f PreSkinningTargetPos = InvSkinningTransform.TransformPosition(GeomCacheVertexPos);
+						Delta = PreSkinningTargetPos - UnskinnedPosition;
+
+						if (bDebugDraw)
+						{
+							OutDebugVectors[SkinnedVertexIndex] = UnskinnedPosition;
+							OutDebugVectors2[SkinnedVertexIndex] = UnskinnedPosition + Delta;
+						}
 					}
 				}
 				else // We're post skinning.
