[MLDeformer] Changes:
- Removed delta mode, always use pre-skinning.
- Always use post-skinning for delta visualization.
- Optimizations, preventing geometry cache evaluations when not needed.
- Deleted some deformer graphs and test assets that are not used anymore.
- Rearranged and regrouped some visualization settings.
- Added a heatmap mode (doesn't work yet, always uses Activation mode for now).
- Added a heatmap scale.
- Some small code refactors.
- Renamed preskinning deformer to DefaultMLDeformerGraph.
- Retrained the cylinder.

diff --git a/Engine/Plugins/Experimental/MLDeformer/Source/MLDeformer/Private/MLDeformerGraphDataInterface.cpp b/Engine/Plugins/Experimental/MLDeformer/Source/MLDeformer/Private/MLDeformerGraphDataInterface.cpp
--- a/Engine/Plugins/Experimental/MLDeformer/Source/MLDeformer/Private/MLDeformerGraphDataInterface.cpp
+++ b/Engine/Plugins/Experimental/MLDeformer/Source/MLDeformer/Private/MLDeformerGraphDataInterface.cpp
@@ -119,9 +119,9 @@ FMLDeformerDataProviderProxy::FMLDeformerDataProviderProxy(UMLDeformerComponent*
 	VertexDeltaMean = DeformerAsset->GetVertexDeltaMean();
 	bCanRunNeuralNet = NeuralNetwork ? (NeuralNetwork->GetInputTensor().Num() == static_cast<int64>(DeformerAsset->GetInputInfo().CalcNumNeuralNetInputs())) : false;
 
-#ifdef WITH_EDITOR
 	VertexDeltaMultiplier = DeformerComponent->GetVertexDeltaMultiplier();
-	DebugScale = 5.f; // todo: Route configurable value from to here from UI.
+#if WITH_EDITORONLY_DATA
+	HeatMapScale = DeformerAsset->GetVizSettings()->GetHeatMapScale();
 #endif
 }
 
@@ -149,7 +149,7 @@ void FMLDeformerDataProviderProxy::GetBindings(int32 InvocationIndex, TCHAR cons
 	Parameters.VertexDeltaScale = VertexDeltaScale;
 	Parameters.VertexDeltaMean = VertexDeltaMean;
 	Parameters.VertexDeltaMultiplier = VertexDeltaMultiplier;
-	Parameters.DebugScale = DebugScale;
+	Parameters.DebugScale = HeatMapScale;
 	Parameters.PositionDeltaBuffer = BufferSRV;
 	Parameters.VertexMapBuffer = VertexMapBufferSRV;
 
diff --git a/Engine/Plugins/Experimental/MLDeformer/Source/MLDeformer/Public/MLDeformerAsset.h b/Engine/Plugins/Experimental/MLDeformer/Source/MLDeformer/Public/MLDeformerAsset.h
--- a/Engine/Plugins/Experimental/MLDeformer/Source/MLDeformer/Public/MLDeformerAsset.h
+++ b/Engine/Plugins/Experimental/MLDeformer/Source/MLDeformer/Public/MLDeformerAsset.h
@@ -41,6 +41,9 @@ UENUM()
 enum class EDecayFunction : uint8
 {
 	Linear,
+
+	// Multiply the learning rate with the Decay Rate each step.
+	// Basically making it smaller all the time.
 	Multiplicative
 };
 
@@ -48,8 +51,11 @@ enum class EDecayFunction : uint8
 UENUM()
 enum class EDeviceType : uint8
 {
-	CPU,	/** Train using CPU. */
-	GPU		/** Train using GPU.*/
+	// Train using the CPU. This tends to be slower.
+	CPU,
+
+	// Train using the GPU. This should be the fastest.
+	GPU
 };
 
 UENUM()
@@ -151,8 +157,6 @@ public:
 	void SetInputInfo(const FMLDeformerInputInfo& Input) { InputInfo = Input; }
 
 #if WITH_EDITORONLY_DATA
-	EDeltaMode GetDeltaMode() const { return DeltaMode; }
-
 	const FTransform& GetAlignmentTransform() const { return AlignmentTransform; }
 
 	const USkeletalMesh* GetSkeletalMesh() const { return SkeletalMesh;  }
@@ -297,10 +301,6 @@ public:
 	UPROPERTY(EditAnywhere, BlueprintReadWrite, AdvancedDisplay, Category = "Training Settings", meta = (ClampMin = "0.0", ClampMax = "100.0")) // A percentage.
 	float NoiseAmount = 0.5;
 
-	/** The delta mode, which describes in what space the calculated deltas are. */
-	UPROPERTY(EditAnywhere, BlueprintReadWrite, AdvancedDisplay, Category = "Training Settings")
-	EDeltaMode DeltaMode = EDeltaMode::PreSkinning;
-
 	/** The loss function to use during model training. */
 	UPROPERTY(EditAnywhere, BlueprintReadWrite, AdvancedDisplay, Category = "Training Settings")
 	EDeviceType DeviceType = EDeviceType::GPU;
diff --git a/Engine/Plugins/Experimental/MLDeformer/Source/MLDeformer/Public/MLDeformerGraphDataInterface.h b/Engine/Plugins/Experimental/MLDeformer/Source/MLDeformer/Public/MLDeformerGraphDataInterface.h
--- a/Engine/Plugins/Experimental/MLDeformer/Source/MLDeformer/Public/MLDeformerGraphDataInterface.h
+++ b/Engine/Plugins/Experimental/MLDeformer/Source/MLDeformer/Public/MLDeformerGraphDataInterface.h
@@ -72,5 +72,5 @@ private:
 	FVector VertexDeltaMean = FVector(0.0f, 0.0f, 0.0f);
 	bool bCanRunNeuralNet = false;
 	float VertexDeltaMultiplier = 1.0f;
-	float DebugScale = 1.0f;
+	float HeatMapScale = 1.0f;
 };
diff --git a/Engine/Plugins/Experimental/MLDeformer/Source/MLDeformer/Public/MLDeformerVizSettings.h b/Engine/Plugins/Experimental/MLDeformer/Source/MLDeformer/Public/MLDeformerVizSettings.h
--- a/Engine/Plugins/Experimental/MLDeformer/Source/MLDeformer/Public/MLDeformerVizSettings.h
+++ b/Engine/Plugins/Experimental/MLDeformer/Source/MLDeformer/Public/MLDeformerVizSettings.h
@@ -11,13 +11,28 @@ class UAnimSequence;
 class UComputeGraph;
 class UGeometryCache;
 
+// The visualization mode, which selects whether you want to view the training data, or test your already trained model.
 UENUM()
 enum class EMLDeformerVizMode : uint8
 {
+	// Preview the training data.
 	TrainingData = 0,
+
+	// Preview testing data, used on trained models.
 	TestData
 };
 
+// The heat map mode which selects what the colors of the heatmap represent.
+UENUM()
+enum class EMLDeformerHeatMapMode : uint8
+{
+	// Visualize the error versus the ground truth model. Requires a ground truth model to be setup.
+	GroundTruth = 0,
+
+	// Visualize areas where the deformer is applying corrections.
+	Activations
+};
+
 /**
  * The vizualization settings.
  */
@@ -52,6 +67,8 @@ public:
 	UComputeGraph* GetDeformerGraph() const { return DeformerGraph; }
 	EMLDeformerVizMode GetVisualizationMode() const { return VisualizationMode; }
 	EMLDeformerVizMode GetTempVisualizationMode() const { return TempVisualizationMode; }
+	EMLDeformerHeatMapMode GetHeatMapMode() const { return HeatMapMode; }
+	float GetHeatMapScale() const { return HeatMapScale; }
 	
 	void SetTempVisualizationMode(EMLDeformerVizMode Mode) { TempVisualizationMode = Mode; }
 	void SetDeformerGraph(UComputeGraph* InDeformerGraph) { DeformerGraph = InDeformerGraph; }
@@ -65,39 +82,47 @@ public:
 	EMLDeformerVizMode TempVisualizationMode = EMLDeformerVizMode::TrainingData;	// Workaround for a bug in UX where the combobox triggers a value changed when clicking it.
 
 	/** The animation sequence to play on the skeletal mesh. */
-	UPROPERTY(EditAnywhere, Category = "Testing")
+	UPROPERTY(EditAnywhere, Category = "Test Assets")
 	TObjectPtr<UAnimSequence> TestAnimSequence = nullptr;
 
 	/** The deformer graph to use on the asset editor's deformed test actor. */
-	UPROPERTY(EditAnywhere, Category = "Testing")
+	UPROPERTY(EditAnywhere, Category = "Test Assets")
 	TObjectPtr<UComputeGraph> DeformerGraph = nullptr;
 
 	/** The geometry cache that represents the ground truth of the test anim sequence. */
-	UPROPERTY(EditAnywhere, Category = "Testing")
+	UPROPERTY(EditAnywhere, Category = "Test Assets")
 	TObjectPtr<UGeometryCache> GroundTruth = nullptr;
 
 	/** The scale factor of the ML deformer deltas being applied on top of the linear skinned results. */
-	UPROPERTY(EditAnywhere, Category = "Testing", meta = (ClampMin = "0.0", ClampMax = "5.0"))
+	UPROPERTY(EditAnywhere, Category = "Live Settings", meta = (ClampMin = "0.0", ClampMax = "5.0"))
 	float VertexDeltaMultiplier = 1.0f;
 
 	/** The play speed factor of the test anim sequence. */
-	UPROPERTY(EditAnywhere, Category = "Testing", meta = (ClampMin = "0.0", ClampMax = "2.0"))
+	UPROPERTY(EditAnywhere, Category = "Live Settings", meta = (ClampMin = "0.0", ClampMax = "2.0"))
 	float AnimPlaySpeed = 1.0f;
 
 	/** Show the heat map? This will visualize the active areas of the deformer. */
-	UPROPERTY(EditAnywhere, Category = "Testing")
+	UPROPERTY(EditAnywhere, Category = "Live Settings")
 	bool bShowHeatMap = false;
 
+	/** What should the heatmap visualize? */
+	UPROPERTY(EditAnywhere, Category = "Live Settings", meta = (EditCondition = "bShowHeatMap"))
+	EMLDeformerHeatMapMode HeatMapMode = EMLDeformerHeatMapMode::GroundTruth;
+
+	/** How many units (centimeters) of error should the most intense color represent? */
+	UPROPERTY(EditAnywhere, Category = "Live Settings", meta = (EditCondition = "bShowHeatMap", ClampMin = "0.00001"))
+	float HeatMapScale = 1.0f;
+
 	/** Draw the linear skinned actor? */
-	UPROPERTY(EditAnywhere, Category = "Testing")
+	UPROPERTY(EditAnywhere, Category = "Live Settings")
 	bool bDrawLinearSkinnedActor = true;
 
 	/** Draw the ML Deformed actor? */
-	UPROPERTY(EditAnywhere, Category = "Testing", DisplayName = "Draw ML Deformed Actor")
+	UPROPERTY(EditAnywhere, Category = "Live Settings", DisplayName = "Draw ML Deformed Actor")
 	bool bDrawMLDeformedActor = true;
 
 	/** Draw the ground truth actor? */
-	UPROPERTY(EditAnywhere, Category = "Testing", meta = (EditCondition = "GroundTruth != nullptr"))
+	UPROPERTY(EditAnywhere, Category = "Live Settings", meta = (EditCondition = "GroundTruth != nullptr"))
 	bool bDrawGroundTruthActor = true;
 
 	/** The frame number of the training data to visualize. */
diff --git a/Engine/Plugins/Experimental/MLDeformer/Source/MLDeformerEditor/Private/MLDeformerAssetDetails.cpp b/Engine/Plugins/Experimental/MLDeformer/Source/MLDeformerEditor/Private/MLDeformerAssetDetails.cpp
--- a/Engine/Plugins/Experimental/MLDeformer/Source/MLDeformerEditor/Private/MLDeformerAssetDetails.cpp
+++ b/Engine/Plugins/Experimental/MLDeformer/Source/MLDeformerEditor/Private/MLDeformerAssetDetails.cpp
@@ -281,7 +281,6 @@ void FMLDeformerAssetDetails::CustomizeDetails(class IDetailLayoutBuilder& Detai
 		}
 	}
 
-	SettingsCategoryBuilder.AddProperty(GET_MEMBER_NAME_CHECKED(UMLDeformerAsset, DeltaMode));
 	SettingsCategoryBuilder.AddProperty(GET_MEMBER_NAME_CHECKED(UMLDeformerAsset, DecayFunction));
 	SettingsCategoryBuilder.AddProperty(GET_MEMBER_NAME_CHECKED(UMLDeformerAsset, DecayRate));
 	SettingsCategoryBuilder.AddProperty(GET_MEMBER_NAME_CHECKED(UMLDeformerAsset, LearningRate));
diff --git a/Engine/Plugins/Experimental/MLDeformer/Source/MLDeformerEditor/Private/MLDeformerEditorData.cpp b/Engine/Plugins/Experimental/MLDeformer/Source/MLDeformerEditor/Private/MLDeformerEditorData.cpp
--- a/Engine/Plugins/Experimental/MLDeformer/Source/MLDeformerEditor/Private/MLDeformerEditorData.cpp
+++ b/Engine/Plugins/Experimental/MLDeformer/Source/MLDeformerEditor/Private/MLDeformerEditorData.cpp
@@ -293,7 +293,10 @@ void FMLDeformerEditorData::InitAssets()
 	InitSettings.DeformerAsset = MLDeformerAsset.Get();
 	InitSettings.World = GetWorld();
 	InitSettings.bLogCacheStats = false;
+	InitSettings.DeltaMode = EDeltaMode::PostSkinning;
 	SingleFrameCache.Init(InitSettings);
+
+	UpdateIsReadyForTrainingState();
 }
 
 void FMLDeformerEditorData::SetAnimFrame(int32 FrameNumber)
@@ -326,36 +329,38 @@ void FMLDeformerEditorData::SetAnimFrame(int32 FrameNumber)
 	CurrentFrame = FrameNumber;
 }
 
-bool FMLDeformerEditorData::IsReadyForTraining() const
+void FMLDeformerEditorData::UpdateIsReadyForTrainingState()
 {
 	UMLDeformerAsset* Asset = GetDeformerAsset();
 
+	bIsReadyForTraining = false;
+
 	// Make sure we have picked required assets.
 	if (Asset == nullptr ||
 		Asset->GetGeometryCache() == nullptr ||
 		GetDeformerAsset()->GetAnimSequence() == nullptr ||
 		GetDeformerAsset()->GetSkeletalMesh() == nullptr)
 	{
-		return false;
+		return;
 	}
 
 	// There are no training frames.
 	if (GetDeformerAsset()->GetNumFrames() == 0)
 	{
-		return false;
+		return;
 	}
 
 	// Now make sure the assets are compatible.
 	if (!Asset->GetVertexErrorText(Asset->GetSkeletalMesh(), Asset->GetGeometryCache(), FText(), FText()).IsEmpty() ||
 		!Asset->GetGeomCacheErrorText(Asset->GetGeometryCache()).IsEmpty())
 	{
-		return false;
+		return;
 	}
 
 	// Make sure we have inputs.
 	if (Asset->CreateInputInfo().IsEmpty())
 	{
-		return false;
+		return;
 	}
 
 	// Make sure every skeletal imported mesh has some geometry track.
@@ -374,11 +379,16 @@ bool FMLDeformerEditorData::IsReadyForTraining() const
 	{
 		if (!GetSingleFrameCache().GetSampler().GetFailedImportedMeshNames().IsEmpty())
 		{
-			return false;
+			return;
 		}
 	}
 
-	return true;
+	bIsReadyForTraining = true;
+}
+
+bool FMLDeformerEditorData::IsReadyForTraining() const
+{
+	return bIsReadyForTraining;
 }
 
 bool FMLDeformerEditorData::GenerateDeltas(uint32 LODIndex, uint32 FrameNumber, TArray<float>& OutDeltas)
@@ -704,7 +714,7 @@ bool FMLDeformerEditorData::IsPlayingAnim() const
 
 FString FMLDeformerEditorData::GetDefaultDeformerGraphAssetPath()
 {
-	return FString(TEXT("/MLDeformer/Deformers/PreSkinningMLDeformerGraph.PreSkinningMLDeformerGraph"));
+	return FString(TEXT("/MLDeformer/Deformers/DefaultMLDeformerGraph.DefaultMLDeformerGraph"));
 }
 
 UComputeGraph* FMLDeformerEditorData::LoadDefaultDeformerGraph()
diff --git a/Engine/Plugins/Experimental/MLDeformer/Source/MLDeformerEditor/Private/MLDeformerEditorToolkit.cpp b/Engine/Plugins/Experimental/MLDeformer/Source/MLDeformerEditor/Private/MLDeformerEditorToolkit.cpp
--- a/Engine/Plugins/Experimental/MLDeformer/Source/MLDeformerEditor/Private/MLDeformerEditorToolkit.cpp
+++ b/Engine/Plugins/Experimental/MLDeformer/Source/MLDeformerEditor/Private/MLDeformerEditorToolkit.cpp
@@ -106,6 +106,7 @@ void FMLDeformerEditorToolkit::InitAssetEditor(
 
 	OnSwitchedVisualizationMode();
 	EditorData->UpdateTimeSlider();
+	EditorData->UpdateIsReadyForTrainingState();
 }
 
 TSharedRef<IPersonaToolkit> FMLDeformerEditorToolkit::GetPersonaToolkit() const 
@@ -187,6 +188,7 @@ void FMLDeformerEditorToolkit::FillToolbar(FToolBarBuilder& ToolbarBuilder)
 						FrameCacheInitSettings.DeformerAsset = DeformerAsset;
 						FrameCacheInitSettings.CacheSizeInBytes = 1024ull * 1024 * DeformerAsset->GetCacheSizeInMegabytes();
 						FrameCacheInitSettings.World = EditorData->GetWorld();
+						FrameCacheInitSettings.DeltaMode = EDeltaMode::PreSkinning;
 						TSharedPtr<FMLDeformerFrameCache> FrameCache = MakeShared<FMLDeformerFrameCache>();
 						FrameCache->Init(FrameCacheInitSettings);
 
@@ -512,7 +514,7 @@ UMLDeformerComponent* FMLDeformerEditorToolkit::AddMLDeformerComponentToActor(EM
 	return Component;
 }
 
-void FMLDeformerEditorToolkit::AddDeformerGraphComponentToActor(EMLDeformerEditorActorIndex ActorIndex, UComputeGraph* ComputeGraph) const
+void FMLDeformerEditorToolkit::AddMeshDeformerToActor(EMLDeformerEditorActorIndex ActorIndex, UComputeGraph* ComputeGraph) const
 {
 	FMLDeformerEditorActor& EditorActor = EditorData->GetEditorActor(ActorIndex);
 	USkinnedMeshComponent* SkelMeshComponent = EditorActor.SkelMeshComponent;
@@ -639,7 +641,7 @@ void FMLDeformerEditorToolkit::HandlePreviewSceneCreated(const TSharedRef<IPerso
 	// Create the component with the deformer graph on it.
 	UMLDeformerVizSettings* VizSettings = EditorData->GetDeformerAsset()->GetVizSettings();
 	UComputeGraph* ComputeGraph = VizSettings->GetDeformerGraph();
-	AddDeformerGraphComponentToActor(EMLDeformerEditorActorIndex::DeformedTest, ComputeGraph);
+	AddMeshDeformerToActor(EMLDeformerEditorActorIndex::DeformedTest, ComputeGraph);
 
 	// Create the ground truth actor.
 	const FLinearColor GroundTruthLabelColor = FMLDeformerEditorStyle::Get().GetColor("MLDeformer.GroundTruth.LabelColor");
@@ -707,7 +709,6 @@ void FMLDeformerEditorToolkit::OnFinishedChangingDetails(const FPropertyChangedE
 	else
 	if (Property->GetFName() == GET_MEMBER_NAME_CHECKED(UMLDeformerAsset, DeltaCutoffLength) ||
 	    Property->GetFName() == GET_MEMBER_NAME_CHECKED(UMLDeformerAsset, AlignmentTransform) ||
-	    Property->GetFName() == GET_MEMBER_NAME_CHECKED(UMLDeformerAsset, DeltaMode) ||
 	    Property->GetFName() == GET_MEMBER_NAME_CHECKED(UMLDeformerAsset, MaxTrainingFrames))
 	{
 		EditorData->InitAssets();
@@ -719,6 +720,7 @@ void FMLDeformerEditorToolkit::OnFinishedChangingDetails(const FPropertyChangedE
 		if (DeformerAsset->GetTempTrainingInputs() != DeformerAsset->GetTrainingInputs())
 		{
 			DeformerAsset->SetTempTrainingInputs(DeformerAsset->GetTrainingInputs());
+			EditorData->UpdateIsReadyForTrainingState();
 			EditorData->GetDetailsView()->ForceRefresh();
 		}
 	}
@@ -730,7 +732,6 @@ void FMLDeformerEditorToolkit::OnFinishedChangingDetails(const FPropertyChangedE
 	else
 	if (Property->GetFName() == GET_MEMBER_NAME_CHECKED(UMLDeformerVizSettings, DeformerGraph))
 	{
-		EditorData->UpdateDeformerGraph();
 		SetComputeGraphDataProviders();
 		EditorData->GetVizSettingsDetailsView()->ForceRefresh();
 	}
@@ -816,7 +817,7 @@ FText FMLDeformerEditorToolkit::GetOverlayText() const
 void FMLDeformerEditorToolkit::SetComputeGraphDataProviders() const
 {
 	UMLDeformerVizSettings* VizSettings = EditorData->GetDeformerAsset()->GetVizSettings();
-	AddDeformerGraphComponentToActor(EMLDeformerEditorActorIndex::DeformedTest, VizSettings->GetDeformerGraph());
+	AddMeshDeformerToActor(EMLDeformerEditorActorIndex::DeformedTest, VizSettings->GetDeformerGraph());
 }
 
 #undef LOCTEXT_NAMESPACE
diff --git a/Engine/Plugins/Experimental/MLDeformer/Source/MLDeformerEditor/Private/MLDeformerFrameCache.cpp b/Engine/Plugins/Experimental/MLDeformer/Source/MLDeformerEditor/Private/MLDeformerFrameCache.cpp
--- a/Engine/Plugins/Experimental/MLDeformer/Source/MLDeformerEditor/Private/MLDeformerFrameCache.cpp
+++ b/Engine/Plugins/Experimental/MLDeformer/Source/MLDeformerEditor/Private/MLDeformerFrameCache.cpp
@@ -152,6 +152,8 @@ void FMLDeformerSamplerData::CalculateVertexDeltas(const TArray<FVector3f>& Skin
 		OutVertexDeltas[Index] = 0.0f;
 	}
 
+	const EDeltaMode DeltaMode = Sampler->GetInitSettings().DeltaMode;
+
 	// For all mesh mappings we found.
 	const float SampleTime = GeometryCacheComponent->GetTimeAtFrame(AnimFrameIndex);
 	const FSkeletalMeshLODModel& LODModel = ImportedModel->LODModels[LODIndex];
@@ -181,7 +183,7 @@ void FMLDeformerSamplerData::CalculateVertexDeltas(const TArray<FVector3f>& Skin
 				FVector3f Delta = FVector3f::ZeroVector;
 
 				const int32 ArrayIndex = 3 * SkinnedVertexIndex;
-				if (DeformerAsset->GetDeltaMode() == EDeltaMode::PreSkinning)
+				if (DeltaMode == EDeltaMode::PreSkinning)
 				{
 					// Calculate the inverse skinning transform for this vertex.
 					const int32 RenderVertexIndex = MeshMapping.ImportedVertexToRenderVertexMap[VertexIndex];
@@ -205,7 +207,7 @@ void FMLDeformerSamplerData::CalculateVertexDeltas(const TArray<FVector3f>& Skin
 				}
 				else // We're post skinning.
 				{
-					check(DeformerAsset->GetDeltaMode() == EDeltaMode::PostSkinning);
+					check(DeltaMode == EDeltaMode::PostSkinning);
 					const FVector3f SkinnedVertexPos = SkinnedPositions[SkinnedVertexIndex];
 					const FVector3f GeomCacheVertexPos = AlignmentTransform.TransformPosition(GeomCacheMeshData.Positions[GeomCacheVertexIndex]);
 					Delta = GeomCacheVertexPos - SkinnedVertexPos;
@@ -449,7 +451,9 @@ void FMLDeformerFrameCache::Init(const FInitSettings& InitSettings)
 
 	DeformerAsset = InitSettings.DeformerAsset;
 
-	if (DeformerAsset->GetSkeletalMesh() == nullptr || DeformerAsset->GetGeometryCache() == nullptr || DeformerAsset->GetAnimSequence() == nullptr)
+	if (DeformerAsset->GetSkeletalMesh() == nullptr || 
+		DeformerAsset->GetGeometryCache() == nullptr || 
+		DeformerAsset->GetAnimSequence() == nullptr)
 	{
 		FrameMap.Empty();
 		CachedTrainingFrames.Empty();
@@ -461,6 +465,7 @@ void FMLDeformerFrameCache::Init(const FInitSettings& InitSettings)
 	FMLDeformerSampler::FInitSettings SamplerInitSettings;
 	SamplerInitSettings.DeformerAsset = InitSettings.DeformerAsset;
 	SamplerInitSettings.World = InitSettings.World;
+	SamplerInitSettings.DeltaMode = InitSettings.DeltaMode;
 	Sampler.Init(SamplerInitSettings);
 
 	// Create the first training frame, and update it already so we know how much memory one frame takes.
diff --git a/Engine/Plugins/Experimental/MLDeformer/Source/MLDeformerEditor/Private/MLDeformerFrameCache.h b/Engine/Plugins/Experimental/MLDeformer/Source/MLDeformerEditor/Private/MLDeformerFrameCache.h
--- a/Engine/Plugins/Experimental/MLDeformer/Source/MLDeformerEditor/Private/MLDeformerFrameCache.h
+++ b/Engine/Plugins/Experimental/MLDeformer/Source/MLDeformerEditor/Private/MLDeformerFrameCache.h
@@ -89,6 +89,7 @@ public:
 	{
 		TObjectPtr<UWorld> World = nullptr;
 		TObjectPtr<UMLDeformerAsset> DeformerAsset = nullptr;
+		EDeltaMode DeltaMode = EDeltaMode::PreSkinning;
 	};
 
 	FMLDeformerSampler() = default;
@@ -170,6 +171,7 @@ public:
 	struct FInitSettings
 	{
 		SIZE_T CacheSizeInBytes = 1024 * 1024 * 1024 * 2ull;	// 2 gigabyte on default.
+		EDeltaMode DeltaMode = EDeltaMode::PreSkinning;
 		TObjectPtr<UMLDeformerAsset> DeformerAsset = nullptr;
 		UWorld* World = nullptr;
 		bool bLogCacheStats = true;
diff --git a/Engine/Plugins/Experimental/MLDeformer/Source/MLDeformerEditor/Private/MLDeformerVizSettingsDetails.cpp b/Engine/Plugins/Experimental/MLDeformer/Source/MLDeformerEditor/Private/MLDeformerVizSettingsDetails.cpp
--- a/Engine/Plugins/Experimental/MLDeformer/Source/MLDeformerEditor/Private/MLDeformerVizSettingsDetails.cpp
+++ b/Engine/Plugins/Experimental/MLDeformer/Source/MLDeformerEditor/Private/MLDeformerVizSettingsDetails.cpp
@@ -81,11 +81,11 @@ void FMLDeformerVizSettingsDetails::CustomizeDetails(class IDetailLayoutBuilder&
 	SharedCategoryBuilder.AddProperty(GET_MEMBER_NAME_CHECKED(UMLDeformerVizSettings, LabelScale));
 	SharedCategoryBuilder.AddProperty(GET_MEMBER_NAME_CHECKED(UMLDeformerVizSettings, MeshSpacing));
 
-	// Testing.
-	IDetailCategoryBuilder& TestingCategoryBuilder = DetailBuilder.EditCategory("Testing", FText::GetEmpty(), ECategoryPriority::Important);
-	TestingCategoryBuilder.SetCategoryVisibility(bShowTestData);
+	// Test Assets.
+	IDetailCategoryBuilder& TestAssetsCategory = DetailBuilder.EditCategory("Test Assets", FText::GetEmpty(), ECategoryPriority::Important);
+	TestAssetsCategory.SetCategoryVisibility(bShowTestData);
 	
-	IDetailPropertyRow& TestAnimRow = TestingCategoryBuilder.AddProperty(GET_MEMBER_NAME_CHECKED(UMLDeformerVizSettings, TestAnimSequence));
+	IDetailPropertyRow& TestAnimRow = TestAssetsCategory.AddProperty(GET_MEMBER_NAME_CHECKED(UMLDeformerVizSettings, TestAnimSequence));
 	TestAnimRow.CustomWidget()
 	.NameContent()
 	[
@@ -108,7 +108,7 @@ void FMLDeformerVizSettingsDetails::CustomizeDetails(class IDetailLayoutBuilder&
 	if (VizSettings)
 	{
 		const FText AnimErrorText = DeformerAsset->GetIncompatibleSkeletonErrorText(DeformerAsset->GetSkeletalMesh(), VizSettings->GetTestAnimSequence());
-		FDetailWidgetRow& AnimErrorRow = TestingCategoryBuilder.AddCustomRow(FText::FromString("AnimSkeletonMisMatchError"))
+		FDetailWidgetRow& AnimErrorRow = TestAssetsCategory.AddCustomRow(FText::FromString("AnimSkeletonMisMatchError"))
 			.Visibility(!AnimErrorText.IsEmpty() ? EVisibility::Visible : EVisibility::Collapsed)
 			.WholeRowContent()
 			[
@@ -122,19 +122,17 @@ void FMLDeformerVizSettingsDetails::CustomizeDetails(class IDetailLayoutBuilder&
 			];
 	}
 
-	TestingCategoryBuilder.AddProperty(GET_MEMBER_NAME_CHECKED(UMLDeformerVizSettings, AnimPlaySpeed));
-
 	FIsResetToDefaultVisible IsResetVisible = FIsResetToDefaultVisible::CreateSP(this, &FMLDeformerVizSettingsDetails::IsResetToDefaultDeformerGraphVisible);
 	FResetToDefaultHandler ResetHandler = FResetToDefaultHandler::CreateSP(this, &FMLDeformerVizSettingsDetails::OnResetToDefaultDeformerGraph);
 	FResetToDefaultOverride ResetOverride = FResetToDefaultOverride::Create(IsResetVisible, ResetHandler);
-	TestingCategoryBuilder.AddProperty(GET_MEMBER_NAME_CHECKED(UMLDeformerVizSettings, DeformerGraph)).OverrideResetToDefault(ResetOverride);
+	TestAssetsCategory.AddProperty(GET_MEMBER_NAME_CHECKED(UMLDeformerVizSettings, DeformerGraph)).OverrideResetToDefault(ResetOverride);
 
 	// Show a warning when no deformer graph has been selected.
 	UObject* Graph = nullptr;
 	TSharedRef<IPropertyHandle> DeformerGraphProperty = DetailBuilder.GetProperty(GET_MEMBER_NAME_CHECKED(UMLDeformerVizSettings, DeformerGraph));
 	if (DeformerGraphProperty->GetValue(Graph) == FPropertyAccess::Result::Success)
 	{
-		FDetailWidgetRow& GraphErrorRow = TestingCategoryBuilder.AddCustomRow(FText::FromString("GraphError"))
+		FDetailWidgetRow& GraphErrorRow = TestAssetsCategory.AddCustomRow(FText::FromString("GraphError"))
 			.Visibility((Graph == nullptr) ? EVisibility::Visible : EVisibility::Collapsed)
 			.WholeRowContent()
 			[
@@ -150,7 +148,7 @@ void FMLDeformerVizSettingsDetails::CustomizeDetails(class IDetailLayoutBuilder&
 
 	if (DeformerAsset)
 	{
-		FDetailWidgetRow& ErrorRow = TestingCategoryBuilder.AddCustomRow(FText::FromString("NoNeuralNetError"))
+		FDetailWidgetRow& ErrorRow = TestAssetsCategory.AddCustomRow(FText::FromString("NoNeuralNetError"))
 			.Visibility((DeformerAsset->GetInferenceNeuralNetwork() == nullptr && Graph != nullptr) ? EVisibility::Visible : EVisibility::Collapsed)
 			.WholeRowContent()
 			[
@@ -164,13 +162,13 @@ void FMLDeformerVizSettingsDetails::CustomizeDetails(class IDetailLayoutBuilder&
 			];
 	}
 
-	TestingCategoryBuilder.AddProperty(GET_MEMBER_NAME_CHECKED(UMLDeformerVizSettings, GroundTruth));
+	TestAssetsCategory.AddProperty(GET_MEMBER_NAME_CHECKED(UMLDeformerVizSettings, GroundTruth));
 
 	// Show an error when the test anim sequence duration doesn't match the one of the ground truth.
 	if (VizSettings)
 	{
 		const FText AnimErrorText = DeformerAsset->GetAnimSequenceErrorText(VizSettings->GroundTruth, VizSettings->GetTestAnimSequence());
-		FDetailWidgetRow& GroundTruthAnimErrorRow = TestingCategoryBuilder.AddCustomRow(FText::FromString("GroundTruthAnimMismatchError"))
+		FDetailWidgetRow& GroundTruthAnimErrorRow = TestAssetsCategory.AddCustomRow(FText::FromString("GroundTruthAnimMismatchError"))
 			.Visibility(!AnimErrorText.IsEmpty() ? EVisibility::Visible : EVisibility::Collapsed)
 			.WholeRowContent()
 			[
@@ -184,7 +182,7 @@ void FMLDeformerVizSettingsDetails::CustomizeDetails(class IDetailLayoutBuilder&
 			];
 
 		const FText GeomErrorText = DeformerAsset->GetGeomCacheErrorText(VizSettings->GetGroundTruth());
-		FDetailWidgetRow& GroundTruthGeomErrorRow = TestingCategoryBuilder.AddCustomRow(FText::FromString("GroundTruthGeomMismatchError"))
+		FDetailWidgetRow& GroundTruthGeomErrorRow = TestAssetsCategory.AddCustomRow(FText::FromString("GroundTruthGeomMismatchError"))
 			.Visibility(!GeomErrorText.IsEmpty() ? EVisibility::Visible : EVisibility::Collapsed)
 			.WholeRowContent()
 			[
@@ -198,7 +196,7 @@ void FMLDeformerVizSettingsDetails::CustomizeDetails(class IDetailLayoutBuilder&
 			];
 
 		const FText VertexErrorText = DeformerAsset->GetVertexErrorText(DeformerAsset->SkeletalMesh, VizSettings->GetGroundTruth(), FText::FromString("Base Mesh"), FText::FromString("Ground Truth Mesh"));
-		FDetailWidgetRow& GroundTruthVertexErrorRow = TestingCategoryBuilder.AddCustomRow(FText::FromString("GroundTruthVertexMismatchError"))
+		FDetailWidgetRow& GroundTruthVertexErrorRow = TestAssetsCategory.AddCustomRow(FText::FromString("GroundTruthVertexMismatchError"))
 			.Visibility(!VertexErrorText.IsEmpty() ? EVisibility::Visible : EVisibility::Collapsed)
 			.WholeRowContent()
 			[
@@ -212,11 +210,21 @@ void FMLDeformerVizSettingsDetails::CustomizeDetails(class IDetailLayoutBuilder&
 			];
 	}
 
-	TestingCategoryBuilder.AddProperty(GET_MEMBER_NAME_CHECKED(UMLDeformerVizSettings, VertexDeltaMultiplier));
-	TestingCategoryBuilder.AddProperty(GET_MEMBER_NAME_CHECKED(UMLDeformerVizSettings, bShowHeatMap));
-	TestingCategoryBuilder.AddProperty(GET_MEMBER_NAME_CHECKED(UMLDeformerVizSettings, bDrawLinearSkinnedActor));
-	TestingCategoryBuilder.AddProperty(GET_MEMBER_NAME_CHECKED(UMLDeformerVizSettings, bDrawMLDeformedActor));
-	TestingCategoryBuilder.AddProperty(GET_MEMBER_NAME_CHECKED(UMLDeformerVizSettings, bDrawGroundTruthActor));
+	IDetailCategoryBuilder& LiveSettingsCategory = DetailBuilder.EditCategory("Live Settings", FText::GetEmpty(), ECategoryPriority::Important);
+	LiveSettingsCategory.SetCategoryVisibility(bShowTestData);
+
+	LiveSettingsCategory.AddProperty(GET_MEMBER_NAME_CHECKED(UMLDeformerVizSettings, VertexDeltaMultiplier));
+	LiveSettingsCategory.AddProperty(GET_MEMBER_NAME_CHECKED(UMLDeformerVizSettings, AnimPlaySpeed));
+
+	IDetailGroup& HeatMapGroup = LiveSettingsCategory.AddGroup("HeatMap", LOCTEXT("HeatMap", "Heat Map"), false, true);
+	HeatMapGroup.AddPropertyRow(DetailBuilder.GetProperty(GET_MEMBER_NAME_CHECKED(UMLDeformerVizSettings, bShowHeatMap)));
+	HeatMapGroup.AddPropertyRow(DetailBuilder.GetProperty(GET_MEMBER_NAME_CHECKED(UMLDeformerVizSettings, HeatMapMode)));
+	HeatMapGroup.AddPropertyRow(DetailBuilder.GetProperty(GET_MEMBER_NAME_CHECKED(UMLDeformerVizSettings, HeatMapScale)));
+
+	IDetailGroup& VisGroup = LiveSettingsCategory.AddGroup("Visibility", LOCTEXT("Visibility", "Visibility"), false, true);
+	VisGroup.AddPropertyRow(DetailBuilder.GetProperty(GET_MEMBER_NAME_CHECKED(UMLDeformerVizSettings, bDrawLinearSkinnedActor)));
+	VisGroup.AddPropertyRow(DetailBuilder.GetProperty(GET_MEMBER_NAME_CHECKED(UMLDeformerVizSettings, bDrawMLDeformedActor)));
+	VisGroup.AddPropertyRow(DetailBuilder.GetProperty(GET_MEMBER_NAME_CHECKED(UMLDeformerVizSettings, bDrawGroundTruthActor)));
 
 	// Training data.
 	IDetailCategoryBuilder& TrainingMeshesCategoryBuilder = DetailBuilder.EditCategory("Training Meshes", FText::GetEmpty(), ECategoryPriority::Important);
diff --git a/Engine/Plugins/Experimental/MLDeformer/Source/MLDeformerEditor/Public/MLDeformerEditorData.h b/Engine/Plugins/Experimental/MLDeformer/Source/MLDeformerEditor/Public/MLDeformerEditorData.h
--- a/Engine/Plugins/Experimental/MLDeformer/Source/MLDeformerEditor/Public/MLDeformerEditorData.h
+++ b/Engine/Plugins/Experimental/MLDeformer/Source/MLDeformerEditor/Public/MLDeformerEditorData.h
@@ -113,6 +113,7 @@ public:
 	void CreateMaterials();
 	void ClampFrameIndex();
 
+	void UpdateIsReadyForTrainingState();
 	bool IsReadyForTraining() const;
 	FText GetOverlayText();
 
@@ -131,6 +132,7 @@ public:
 	TArray<float> VertexDeltas;
 
 	bool bIsVertexDeltaNormalized = false;
+	bool bIsReadyForTraining = false;
 	int32 CurrentFrame = -1;
 
 private:
diff --git a/Engine/Plugins/Experimental/MLDeformer/Source/MLDeformerEditor/Public/MLDeformerEditorToolkit.h b/Engine/Plugins/Experimental/MLDeformer/Source/MLDeformerEditor/Public/MLDeformerEditorToolkit.h
--- a/Engine/Plugins/Experimental/MLDeformer/Source/MLDeformerEditor/Public/MLDeformerEditorToolkit.h
+++ b/Engine/Plugins/Experimental/MLDeformer/Source/MLDeformerEditor/Public/MLDeformerEditorToolkit.h
@@ -90,7 +90,7 @@ private:
 	void CreateBaseActor(const TSharedRef<IPersonaPreviewScene>& InPersonaPreviewScene, const FName& Name, FLinearColor LabelColor, FLinearColor WireframeColor);
 	void CreateGeomCacheActor(EMLDeformerEditorActorIndex ActorIndex, UWorld* World, const FName& Name, UGeometryCache* GeomCache, FLinearColor LabelColor, FLinearColor WireframeColor);
 	UMLDeformerComponent* AddMLDeformerComponentToActor(EMLDeformerEditorActorIndex ActorIndex);
-	void AddDeformerGraphComponentToActor(EMLDeformerEditorActorIndex ActorIndex, UComputeGraph* ComputeGraph) const;
+	void AddMeshDeformerToActor(EMLDeformerEditorActorIndex ActorIndex, UComputeGraph* ComputeGraph) const;
 	bool TryLoadOnnxFile() const;
 	void ShowNotification(const FText& Message, SNotificationItem::ECompletionState State, bool PlaySound) const;
 	UTextRenderComponent* CreateLabelForActor(AActor* Actor, UWorld* World, FLinearColor Color, const FText& Text) const;
