Fix MLDeformer bug with multiple segments.
Bug was that segment vertex buffer offsets weren't respected when reading ML deltas from buffer.
Add MLDeformer debug data interface.
This contains the heat map multiplier and new ground truth position access.
Needs some further work on Editor UI code to expose the ground truth heatmap correctly.

diff --git a/Engine/Plugins/Experimental/MLDeformer/Shaders/Private/MLDeformerDataInterface.ush b/Engine/Plugins/Experimental/MLDeformer/Shaders/Private/MLDeformerDataInterface.ush
new file mode 100644
--- /dev/null
+++ b/Engine/Plugins/Experimental/MLDeformer/Shaders/Private/MLDeformerDataInterface.ush
@@ -0,0 +1,31 @@
+// Copyright Epic Games, Inc. All Rights Reserved.
+
+#include "/Plugin/ComputeFramework/Private/ComputeKernelCommon.ush"
+
+uint DI_LOCAL(NumVertices);
+uint DI_LOCAL(InputStreamStart);
+float3 DI_LOCAL(VertexDeltaScale);
+float3 DI_LOCAL(VertexDeltaMean);
+float DI_LOCAL(VertexDeltaMultiplier);
+float DI_LOCAL(DebugScale);
+
+Buffer<float> DI_LOCAL(PositionDeltaBuffer);
+Buffer<uint> DI_LOCAL(VertexMapBuffer);
+
+DI_IMPL_READ(ReadNumVertices, uint, void)
+{
+	return DI_LOCAL(NumVertices);
+}
+
+DI_IMPL_READ(ReadPositionDelta, float3, uint VertexIndex)
+{
+	uint BufferIndex = DI_LOCAL(VertexMapBuffer)[VertexIndex + DI_LOCAL(InputStreamStart)] * 3;
+	float3 Delta = float3(DI_LOCAL(PositionDeltaBuffer)[BufferIndex], DI_LOCAL(PositionDeltaBuffer)[BufferIndex + 1], DI_LOCAL(PositionDeltaBuffer)[BufferIndex + 2]);
+	float3 FinalDelta = Delta * DI_LOCAL(VertexDeltaScale) + DI_LOCAL(VertexDeltaMean);
+	return FinalDelta * DI_LOCAL(VertexDeltaMultiplier);
+}
+
+DI_IMPL_READ(ReadDebugScale, float, void)
+{
+	return DI_LOCAL(DebugScale);
+}
diff --git a/Engine/Plugins/Experimental/MLDeformer/Shaders/Private/MLDeformerDebugDataInterface.ush b/Engine/Plugins/Experimental/MLDeformer/Shaders/Private/MLDeformerDebugDataInterface.ush
new file mode 100644
--- /dev/null
+++ b/Engine/Plugins/Experimental/MLDeformer/Shaders/Private/MLDeformerDebugDataInterface.ush
@@ -0,0 +1,26 @@
+// Copyright Epic Games, Inc. All Rights Reserved.
+
+#include "/Plugin/ComputeFramework/Private/ComputeKernelCommon.ush"
+
+uint DI_LOCAL(NumVertices);
+uint DI_LOCAL(InputStreamStart);
+float DI_LOCAL(HeatMapScale);
+
+StructuredBuffer<float3> DI_LOCAL(PositionGroundTruthBuffer);
+Buffer<uint> DI_LOCAL(VertexMapBuffer);
+
+DI_IMPL_READ(ReadNumVertices, uint, void)
+{
+	return DI_LOCAL(NumVertices);
+}
+
+DI_IMPL_READ(ReadPositionGroundTruth, float3, uint VertexIndex)
+{
+	uint BufferIndex = DI_LOCAL(VertexMapBuffer)[VertexIndex + DI_LOCAL(InputStreamStart)];
+	return DI_LOCAL(PositionGroundTruthBuffer)[BufferIndex];
+}
+
+DI_IMPL_READ(ReadHeatMapScale, float, void)
+{
+	return DI_LOCAL(HeatMapScale);
+}
diff --git a/Engine/Plugins/Experimental/MLDeformer/Shaders/Private/MLDeformerGraphDataInterface.ush b/Engine/Plugins/Experimental/MLDeformer/Shaders/Private/MLDeformerGraphDataInterface.ush
deleted file mode 100644
diff --git a/Engine/Plugins/Experimental/MLDeformer/Source/MLDeformer/Private/MLDeformerDataInterface.cpp b/Engine/Plugins/Experimental/MLDeformer/Source/MLDeformer/Private/MLDeformerDataInterface.cpp
new file mode 100644
--- /dev/null
+++ b/Engine/Plugins/Experimental/MLDeformer/Source/MLDeformer/Private/MLDeformerDataInterface.cpp
@@ -0,0 +1,160 @@
+// Copyright Epic Games, Inc. All Rights Reserved.
+
+#include "MLDeformerDataInterface.h"
+
+#include "ComputeFramework/ComputeKernelPermutationSet.h"
+#include "ComputeFramework/ShaderParamTypeDefinition.h"
+#include "MLDeformerAsset.h"
+#include "MLDeformerComponent.h"
+#include "NeuralNetwork.h"
+#include "OptimusDataDomain.h"
+#include "RenderGraphBuilder.h"
+#include "RenderGraphUtils.h"
+#include "RenderGraphResources.h"
+#include "ShaderParameterMetadataBuilder.h"
+#include "SkeletalRenderPublic.h"
+
+FString UMLDeformerDataInterface::GetDisplayName() const
+{
+	return TEXT("ML Deformer");
+}
+
+TArray<FOptimusCDIPinDefinition> UMLDeformerDataInterface::GetPinDefinitions() const
+{
+	TArray<FOptimusCDIPinDefinition> Defs;
+	Defs.Add({ "PositionDelta", "ReadPositionDelta", Optimus::DomainName::Vertex, "ReadNumVertices" });
+	return Defs;
+}
+
+void UMLDeformerDataInterface::GetSupportedInputs(TArray<FShaderFunctionDefinition>& OutFunctions) const
+{
+	{
+		FShaderFunctionDefinition Fn;
+		Fn.Name = TEXT("ReadNumVertices");
+		Fn.bHasReturnType = true;
+		FShaderParamTypeDefinition ReturnParam = {};
+		ReturnParam.ValueType = FShaderValueType::Get(EShaderFundamentalType::Uint);
+		Fn.ParamTypes.Add(ReturnParam);
+		OutFunctions.Add(Fn);
+	}
+	{
+		FShaderFunctionDefinition Fn;
+		Fn.Name = TEXT("ReadPositionDelta");
+		Fn.bHasReturnType = true;
+		FShaderParamTypeDefinition ReturnParam = {};
+		ReturnParam.ValueType = FShaderValueType::Get(EShaderFundamentalType::Float, 3);
+		Fn.ParamTypes.Add(ReturnParam);
+		FShaderParamTypeDefinition Param0 = {};
+		Param0.ValueType = FShaderValueType::Get(EShaderFundamentalType::Uint);
+		Fn.ParamTypes.Add(Param0);
+		OutFunctions.Add(Fn);
+	}
+}
+
+BEGIN_SHADER_PARAMETER_STRUCT(FMLDeformerDataInterfaceParameters, )
+	SHADER_PARAMETER(uint32, NumVertices)
+	SHADER_PARAMETER(uint32, InputStreamStart)
+	SHADER_PARAMETER(FVector3f, VertexDeltaScale)
+	SHADER_PARAMETER(FVector3f, VertexDeltaMean)
+	SHADER_PARAMETER(float, VertexDeltaMultiplier)
+	SHADER_PARAMETER_RDG_BUFFER_SRV(Buffer<float>, PositionDeltaBuffer)
+	SHADER_PARAMETER_SRV(Buffer<uint>, VertexMapBuffer)
+END_SHADER_PARAMETER_STRUCT()
+
+void UMLDeformerDataInterface::GetShaderParameters(TCHAR const* UID, FShaderParametersMetadataBuilder& OutBuilder) const
+{
+	OutBuilder.AddNestedStruct<FMLDeformerDataInterfaceParameters>(UID);
+}
+
+void UMLDeformerDataInterface::GetHLSL(FString& OutHLSL) const
+{
+	OutHLSL += TEXT("#include \"/Plugin/MLDeformer/Private/MLDeformerDataInterface.ush\"\n");
+}
+
+void UMLDeformerDataInterface::GetSourceTypes(TArray<UClass*>& OutSourceTypes) const
+{
+	OutSourceTypes.Add(USkeletalMeshComponent::StaticClass());
+	OutSourceTypes.Add(UMLDeformerComponent::StaticClass());
+}
+
+UComputeDataProvider* UMLDeformerDataInterface::CreateDataProvider(TArrayView<TObjectPtr<UObject>> InSourceObjects, uint64 InInputMask, uint64 InOutputMask) const
+{
+	UMLDeformerDataProvider* Provider = NewObject<UMLDeformerDataProvider>();
+	if (InSourceObjects.Num() == 2)
+	{
+		Provider->SkeletalMeshComponent = Cast<USkeletalMeshComponent>(InSourceObjects[0]);
+		Provider->DeformerComponent = Cast<UMLDeformerComponent>(InSourceObjects[1]);
+	}
+	return Provider;
+}
+
+bool UMLDeformerDataProvider::IsValid() const
+{
+	return
+		SkeletalMeshComponent != nullptr &&
+		SkeletalMeshComponent->MeshObject != nullptr &&
+		DeformerComponent != nullptr &&
+		DeformerComponent->GetDeformerAsset() != nullptr &&
+		DeformerComponent->GetDeformerAsset()->GetVertexMapBuffer().ShaderResourceViewRHI != nullptr &&
+		DeformerComponent->GetDeformerAsset()->GetInferenceNeuralNetwork() != nullptr &&
+		DeformerComponent->GetDeformerAsset()->GetInferenceNeuralNetwork()->IsLoaded();
+}
+
+FComputeDataProviderRenderProxy* UMLDeformerDataProvider::GetRenderProxy()
+{
+	return new FMLDeformerDataProviderProxy(SkeletalMeshComponent, DeformerComponent);
+}
+
+FMLDeformerDataProviderProxy::FMLDeformerDataProviderProxy(USkeletalMeshComponent* SkeletalMeshComponent, UMLDeformerComponent* DeformerComponent)
+{
+	SkeletalMeshObject = SkeletalMeshComponent->MeshObject;
+
+	const UMLDeformerAsset* DeformerAsset = DeformerComponent->GetDeformerAsset();
+	NeuralNetwork = DeformerAsset->GetInferenceNeuralNetwork();
+	bCanRunNeuralNet = NeuralNetwork ? (NeuralNetwork->GetInputTensor().Num() == static_cast<int64>(DeformerAsset->GetInputInfo().CalcNumNeuralNetInputs())) : false;
+	
+	VertexMapBufferSRV = DeformerAsset->GetVertexMapBuffer().ShaderResourceViewRHI;
+	VertexDeltaScale = DeformerAsset->GetVertexDeltaScale();
+	VertexDeltaMean = DeformerAsset->GetVertexDeltaMean();
+
+	VertexDeltaMultiplier = DeformerComponent->GetVertexDeltaMultiplier();
+}
+
+void FMLDeformerDataProviderProxy::AllocateResources(FRDGBuilder& GraphBuilder)
+{
+	if (bCanRunNeuralNet)
+	{
+		check(NeuralNetwork != nullptr);
+		Buffer = GraphBuilder.RegisterExternalBuffer(NeuralNetwork->GetOutputTensor().GetPooledBuffer());
+	}
+	else
+	{
+		// TODO: use an actual buffer that is of the right size, and filled with zero's.
+		Buffer = GraphBuilder.RegisterExternalBuffer(GWhiteVertexBufferWithRDG->Buffer);
+	}
+
+	BufferSRV = GraphBuilder.CreateSRV(Buffer, PF_R32_FLOAT);
+}
+
+void FMLDeformerDataProviderProxy::GetBindings(int32 InvocationIndex, TCHAR const* UID, FBindings& OutBindings) const
+{
+	const int32 SectionIdx = InvocationIndex;
+	FSkeletalMeshRenderData const& SkeletalMeshRenderData = SkeletalMeshObject->GetSkeletalMeshRenderData();
+	FSkeletalMeshLODRenderData const* LodRenderData = SkeletalMeshRenderData.GetPendingFirstLOD(0);
+	FSkelMeshRenderSection const& RenderSection = LodRenderData->RenderSections[SectionIdx];
+
+	FMLDeformerDataInterfaceParameters Parameters;
+	FMemory::Memset(&Parameters, 0, sizeof(Parameters));
+	Parameters.NumVertices = 0;
+	Parameters.InputStreamStart = RenderSection.BaseVertexIndex;
+	Parameters.VertexDeltaScale = VertexDeltaScale;
+	Parameters.VertexDeltaMean = VertexDeltaMean;
+	Parameters.VertexDeltaMultiplier = VertexDeltaMultiplier;
+	Parameters.PositionDeltaBuffer = BufferSRV;
+	Parameters.VertexMapBuffer = VertexMapBufferSRV;
+
+	TArray<uint8> ParamData;
+	ParamData.SetNum(sizeof(Parameters));
+	FMemory::Memcpy(ParamData.GetData(), &Parameters, sizeof(Parameters));
+	OutBindings.Structs.Add(TTuple<FString, TArray<uint8>>(UID, MoveTemp(ParamData)));
+}
diff --git a/Engine/Plugins/Experimental/MLDeformer/Source/MLDeformer/Private/MLDeformerDebugDataInterface.cpp b/Engine/Plugins/Experimental/MLDeformer/Source/MLDeformer/Private/MLDeformerDebugDataInterface.cpp
new file mode 100644
--- /dev/null
+++ b/Engine/Plugins/Experimental/MLDeformer/Source/MLDeformer/Private/MLDeformerDebugDataInterface.cpp
@@ -0,0 +1,240 @@
+// Copyright Epic Games, Inc. All Rights Reserved.
+
+#include "MLDeformerDebugDataInterface.h"
+
+#include "Components/SkeletalMeshComponent.h"
+#include "ComputeFramework/ComputeKernelPermutationSet.h"
+#include "ComputeFramework/ShaderParamTypeDefinition.h"
+#include "GeometryCache.h"
+#include "GeometryCacheMeshData.h"
+#include "GeometryCacheTrack.h"
+#include "MLDeformerAsset.h"
+#include "MLDeformerComponent.h"
+#include "OptimusDataDomain.h"
+#include "RenderGraphBuilder.h"
+#include "RenderGraphUtils.h"
+#include "RenderGraphResources.h"
+#include "Rendering/SkeletalMeshModel.h"
+#include "ShaderParameterMetadataBuilder.h"
+#include "SkeletalRenderPublic.h"
+
+FString UMLDeformerDebugDataInterface::GetDisplayName() const
+{
+	return TEXT("ML Deformer Debug");
+}
+
+TArray<FOptimusCDIPinDefinition> UMLDeformerDebugDataInterface::GetPinDefinitions() const
+{
+	TArray<FOptimusCDIPinDefinition> Defs;
+	Defs.Add({ "HeatMapScale", "ReadHeatMapScale" });
+	Defs.Add({ "PositionGroundTruth", "ReadPositionGroundTruth", Optimus::DomainName::Vertex, "ReadNumVertices" });
+	return Defs;
+}
+
+void UMLDeformerDebugDataInterface::GetSupportedInputs(TArray<FShaderFunctionDefinition>& OutFunctions) const
+{
+	{
+		FShaderFunctionDefinition Fn;
+		Fn.Name = TEXT("ReadNumVertices");
+		Fn.bHasReturnType = true;
+		FShaderParamTypeDefinition ReturnParam = {};
+		ReturnParam.ValueType = FShaderValueType::Get(EShaderFundamentalType::Uint);
+		Fn.ParamTypes.Add(ReturnParam);
+		OutFunctions.Add(Fn);
+	}
+	{
+		FShaderFunctionDefinition Fn;
+		Fn.Name = TEXT("ReadHeatMapScale");
+		Fn.bHasReturnType = true;
+		FShaderParamTypeDefinition ReturnParam = {};
+		ReturnParam.ValueType = FShaderValueType::Get(EShaderFundamentalType::Float);
+		Fn.ParamTypes.Add(ReturnParam);
+		OutFunctions.Add(Fn);
+	}
+	{
+		FShaderFunctionDefinition Fn;
+		Fn.Name = TEXT("ReadPositionGroundTruth");
+		Fn.bHasReturnType = true;
+		FShaderParamTypeDefinition ReturnParam = {};
+		ReturnParam.ValueType = FShaderValueType::Get(EShaderFundamentalType::Float, 3);
+		Fn.ParamTypes.Add(ReturnParam);
+		FShaderParamTypeDefinition Param0 = {};
+		Param0.ValueType = FShaderValueType::Get(EShaderFundamentalType::Uint);
+		Fn.ParamTypes.Add(Param0);
+		OutFunctions.Add(Fn);
+	}
+}
+
+BEGIN_SHADER_PARAMETER_STRUCT(FMLDeformerDebugDataInterfaceParameters, )
+	SHADER_PARAMETER(uint32, NumVertices)
+	SHADER_PARAMETER(uint32, InputStreamStart)
+	SHADER_PARAMETER(float, HeatMapScale)
+	SHADER_PARAMETER_RDG_BUFFER_SRV(StructuredBuffer<float3>, PositionGroundTruthBuffer)
+	SHADER_PARAMETER_SRV(Buffer<uint>, VertexMapBuffer)
+END_SHADER_PARAMETER_STRUCT()
+
+void UMLDeformerDebugDataInterface::GetShaderParameters(TCHAR const* UID, FShaderParametersMetadataBuilder& OutBuilder) const
+{
+	OutBuilder.AddNestedStruct<FMLDeformerDebugDataInterfaceParameters>(UID);
+}
+
+void UMLDeformerDebugDataInterface::GetHLSL(FString& OutHLSL) const
+{
+	OutHLSL += TEXT("#include \"/Plugin/MLDeformer/Private/MLDeformerDebugDataInterface.ush\"\n");
+}
+
+void UMLDeformerDebugDataInterface::GetSourceTypes(TArray<UClass*>& OutSourceTypes) const
+{
+	OutSourceTypes.Add(USkeletalMeshComponent::StaticClass());
+	OutSourceTypes.Add(UMLDeformerComponent::StaticClass());
+}
+
+UComputeDataProvider* UMLDeformerDebugDataInterface::CreateDataProvider(TArrayView<TObjectPtr<UObject>> InSourceObjects, uint64 InInputMask, uint64 InOutputMask) const
+{
+	UMLDeformerDebugDataProvider* Provider = NewObject<UMLDeformerDebugDataProvider>();
+	if (InSourceObjects.Num() == 2)
+	{
+		Provider->SkeletalMeshComponent = Cast<USkeletalMeshComponent>(InSourceObjects[0]);
+		Provider->DeformerComponent = Cast<UMLDeformerComponent>(InSourceObjects[1]);
+	}
+#if WITH_EDITORONLY_DATA
+	UMLDeformerAsset* DeformerAsset = Provider->DeformerComponent != nullptr ? Provider->DeformerComponent->GetDeformerAsset() : nullptr;
+	if (DeformerAsset != nullptr)
+	{
+		TArray<FString> FailedImportedMeshNames;
+		UMLDeformerAsset::GenerateMeshMappings(DeformerAsset, Provider->MeshMappings, FailedImportedMeshNames);
+	}
+#endif
+	return Provider;
+}
+
+bool UMLDeformerDebugDataProvider::IsValid() const
+{
+#if WITH_EDITORONLY_DATA
+	return
+		SkeletalMeshComponent != nullptr &&
+		SkeletalMeshComponent->MeshObject != nullptr &&
+		DeformerComponent != nullptr &&
+		DeformerComponent->GetDeformerAsset() != nullptr &&
+		DeformerComponent->GetDeformerAsset()->GetVertexMapBuffer().ShaderResourceViewRHI != nullptr;
+#else
+	return false; // This data interface is only valid in editor.
+#endif
+}
+
+FComputeDataProviderRenderProxy* UMLDeformerDebugDataProvider::GetRenderProxy()
+{
+#if WITH_EDITORONLY_DATA
+	return new FMLDeformerDebugDataProviderProxy(SkeletalMeshComponent, DeformerComponent, MeshMappings);
+#else
+	return nullptr;
+#endif
+}
+
+#if WITH_EDITORONLY_DATA
+
+void GetGroundTruthPositions(
+	int32 LODIndex,
+	float SampleTime,
+	UMLDeformerAsset const* DeformerAsset,
+	TArray<FMLDeformerMeshMapping> const& MeshMappings,
+	TArray<FVector3f>& OutPositions)
+{
+	if (!ensure(DeformerAsset))
+	{
+		return;
+	}
+	USkeletalMesh const* SkelMesh = DeformerAsset->GetSkeletalMesh();
+	if (!ensure(SkelMesh))
+	{
+		return;
+	}
+	FSkeletalMeshModel const* ImportedModel = SkelMesh->GetImportedModel();
+	if (!ensure(ImportedModel))
+	{
+		return;
+	}
+	UGeometryCache const* GeomCache = DeformerAsset->GetGeometryCache();
+	if (!ensure(GeomCache))
+	{
+		return;
+	}
+
+	const FTransform AlignmentTransform = DeformerAsset->GetAlignmentTransform();
+
+	FSkeletalMeshLODModel const& LODModel = ImportedModel->LODModels[LODIndex];
+	TArray<FSkelMeshImportedMeshInfo> const& SkelMeshInfos = LODModel.ImportedMeshInfos;
+
+	const uint32 NumVertices = LODModel.MaxImportVertex + 1;
+	OutPositions.Reset(NumVertices);
+	OutPositions.AddZeroed(NumVertices);
+
+	// For all mesh mappings we found.
+	for (int32 MeshMappingIndex = 0; MeshMappingIndex < MeshMappings.Num(); ++MeshMappingIndex)
+	{
+		FMLDeformerMeshMapping const& MeshMapping = MeshMappings[MeshMappingIndex];
+		FSkelMeshImportedMeshInfo const& MeshInfo = SkelMeshInfos[MeshMapping.MeshIndex];
+		UGeometryCacheTrack* Track = GeomCache->Tracks[MeshMapping.TrackIndex];
+
+		FGeometryCacheMeshData GeomCacheMeshData;
+		if (!Track->GetMeshDataAtTime(SampleTime, GeomCacheMeshData))
+		{
+			continue;
+		}
+
+		for (int32 VertexIndex = 0; VertexIndex < MeshInfo.NumVertices; ++VertexIndex)
+		{
+			const int32 SkinnedVertexIndex = MeshInfo.StartImportedVertex + VertexIndex;
+			const int32 GeomCacheVertexIndex = MeshMapping.SkelMeshToTrackVertexMap[VertexIndex];
+			if (GeomCacheVertexIndex != INDEX_NONE && GeomCacheMeshData.Positions.IsValidIndex(GeomCacheVertexIndex))
+			{
+				const FVector3f GeomCacheVertexPos = AlignmentTransform.TransformPosition(GeomCacheMeshData.Positions[GeomCacheVertexIndex]);
+				OutPositions[SkinnedVertexIndex] = GeomCacheVertexPos;
+			}
+		}
+	}
+}
+
+FMLDeformerDebugDataProviderProxy::FMLDeformerDebugDataProviderProxy(USkeletalMeshComponent* SkeletalMeshComponent, UMLDeformerComponent* DeformerComponent, TArray<FMLDeformerMeshMapping> const& MeshMappings)
+{
+	SkeletalMeshObject = SkeletalMeshComponent->MeshObject;
+	
+	UMLDeformerAsset const* DeformerAsset = DeformerComponent->GetDeformerAsset();
+	VertexMapBufferSRV = DeformerAsset->GetVertexMapBuffer().ShaderResourceViewRHI;
+	HeatMapScale = DeformerAsset->GetVizSettings()->GetHeatMapScale();
+
+	const int32 LODIndex = 0;
+	const float SampleTime = SkeletalMeshComponent->GetPosition();
+	GetGroundTruthPositions(0, SampleTime, DeformerAsset, MeshMappings, GroundTruthPositions);
+}
+
+#endif // WITH_EDITORONLY_DATA
+
+void FMLDeformerDebugDataProviderProxy::AllocateResources(FRDGBuilder& GraphBuilder)
+{
+	GroundTruthBuffer = GraphBuilder.CreateBuffer(FRDGBufferDesc::CreateStructuredDesc(sizeof(FVector3f), GroundTruthPositions.Num()), TEXT("MLDeformer.GroundTruthPositions"));
+	GroundTruthBufferSRV = GraphBuilder.CreateSRV(GroundTruthBuffer);
+
+	GraphBuilder.QueueBufferUpload(GroundTruthBuffer, GroundTruthPositions.GetData(), GroundTruthPositions.GetAllocatedSize(), ERDGInitialDataFlags::None);
+}
+
+void FMLDeformerDebugDataProviderProxy::GetBindings(int32 InvocationIndex, TCHAR const* UID, FBindings& OutBindings) const
+{
+	const int32 SectionIdx = InvocationIndex;
+	FSkeletalMeshRenderData const& SkeletalMeshRenderData = SkeletalMeshObject->GetSkeletalMeshRenderData();
+	FSkeletalMeshLODRenderData const* LodRenderData = SkeletalMeshRenderData.GetPendingFirstLOD(0);
+	FSkelMeshRenderSection const& RenderSection = LodRenderData->RenderSections[SectionIdx];
+
+	FMLDeformerDebugDataInterfaceParameters Parameters;
+	FMemory::Memset(&Parameters, 0, sizeof(Parameters));
+	Parameters.NumVertices = 0;
+	Parameters.InputStreamStart = RenderSection.BaseVertexIndex;
+	Parameters.HeatMapScale = HeatMapScale;
+	Parameters.PositionGroundTruthBuffer = GroundTruthBufferSRV;
+	Parameters.VertexMapBuffer = VertexMapBufferSRV;
+
+	TArray<uint8> ParamData;
+	ParamData.SetNum(sizeof(Parameters));
+	FMemory::Memcpy(ParamData.GetData(), &Parameters, sizeof(Parameters));
+	OutBindings.Structs.Add(TTuple<FString, TArray<uint8>>(UID, MoveTemp(ParamData)));
+}
diff --git a/Engine/Plugins/Experimental/MLDeformer/Source/MLDeformer/Private/MLDeformerGraphDataInterface.cpp b/Engine/Plugins/Experimental/MLDeformer/Source/MLDeformer/Private/MLDeformerGraphDataInterface.cpp
deleted file mode 100644
diff --git a/Engine/Plugins/Experimental/MLDeformer/Source/MLDeformer/Public/MLDeformerDataInterface.h b/Engine/Plugins/Experimental/MLDeformer/Source/MLDeformer/Public/MLDeformerDataInterface.h
new file mode 100644
--- /dev/null
+++ b/Engine/Plugins/Experimental/MLDeformer/Source/MLDeformer/Public/MLDeformerDataInterface.h
@@ -0,0 +1,81 @@
+// Copyright Epic Games, Inc. All Rights Reserved.
+
+#pragma once
+
+#include "MLDeformer.h"
+#include "MLDeformerAsset.h"
+#include "OptimusComputeDataInterface.h"
+#include "ComputeFramework/ComputeDataProvider.h"
+#include "MLDeformerDataInterface.generated.h"
+
+class FRDGBuffer;
+class FRDGBufferSRV;
+class FRHIShaderResourceView;
+class FSkeletalMeshObject;
+class UMLDeformerComponent;
+class UNeuralNetwork;
+class USkeletalMeshComponent;
+
+/** Compute Framework Data Interface for MLDefomer data. */
+UCLASS(Category = ComputeFramework)
+class MLDEFORMER_API UMLDeformerDataInterface : public UOptimusComputeDataInterface
+{
+	GENERATED_BODY()
+
+public:
+	//~ Begin UOptimusComputeDataInterface Interface
+	FString GetDisplayName() const override;
+	TArray<FOptimusCDIPinDefinition> GetPinDefinitions() const override;
+	//~ End UOptimusComputeDataInterface Interface
+
+	//~ Begin UComputeDataInterface Interface
+	void GetSupportedInputs(TArray<FShaderFunctionDefinition>& OutFunctions) const override;
+	void GetShaderParameters(TCHAR const* UID, FShaderParametersMetadataBuilder& OutBuilder) const override;
+	void GetHLSL(FString& OutHLSL) const override;
+	void GetSourceTypes(TArray<UClass*>& OutSourceTypes) const override;
+	UComputeDataProvider* CreateDataProvider(TArrayView<TObjectPtr<UObject>> InSourceObjects, uint64 InInputMask, uint64 InOutputMask) const override;
+	//~ End UComputeDataInterface Interface
+};
+
+
+/** Compute Framework Data Provider for MLDeformer data. */
+UCLASS(BlueprintType, editinlinenew, Category = ComputeFramework)
+class MLDEFORMER_API UMLDeformerDataProvider : public UComputeDataProvider
+{
+	GENERATED_BODY()
+
+public:
+	UPROPERTY(EditAnywhere, BlueprintReadWrite, Category = Binding)
+	TObjectPtr<USkeletalMeshComponent> SkeletalMeshComponent = nullptr;
+	UPROPERTY(EditAnywhere, BlueprintReadWrite, Category = Binding)
+	TObjectPtr<UMLDeformerComponent> DeformerComponent = nullptr;
+
+	//~ Begin UComputeDataProvider Interface
+	bool IsValid() const override;
+	FComputeDataProviderRenderProxy* GetRenderProxy() override;
+	//~ End UComputeDataProvider Interface
+};
+
+
+/** Compute Framework Data Provider Proxy for MLDeformer data. */
+class FMLDeformerDataProviderProxy : public FComputeDataProviderRenderProxy
+{
+public:
+	FMLDeformerDataProviderProxy(USkeletalMeshComponent* SkeletalMeshComponent, UMLDeformerComponent* DeformerComponent);
+
+	//~ Begin FComputeDataProviderRenderProxy Interface
+	void AllocateResources(FRDGBuilder& GraphBuilder) override;
+	void GetBindings(int32 InvocationIndex, TCHAR const* UID, FBindings& OutBindings) const override;
+	//~ End FComputeDataProviderRenderProxy Interface
+
+private:
+	FSkeletalMeshObject* SkeletalMeshObject;
+	UNeuralNetwork* NeuralNetwork = nullptr;
+	bool bCanRunNeuralNet = false;
+	FRHIShaderResourceView* VertexMapBufferSRV = nullptr;
+	FRDGBuffer* Buffer = nullptr;
+	FRDGBufferSRV* BufferSRV = nullptr;
+	FVector VertexDeltaScale = FVector(1.0f, 1.0f, 1.0f);
+	FVector VertexDeltaMean = FVector(0.0f, 0.0f, 0.0f);
+	float VertexDeltaMultiplier = 1.0f;
+};
diff --git a/Engine/Plugins/Experimental/MLDeformer/Source/MLDeformer/Public/MLDeformerDebugDataInterface.h b/Engine/Plugins/Experimental/MLDeformer/Source/MLDeformer/Public/MLDeformerDebugDataInterface.h
new file mode 100644
--- /dev/null
+++ b/Engine/Plugins/Experimental/MLDeformer/Source/MLDeformer/Public/MLDeformerDebugDataInterface.h
@@ -0,0 +1,91 @@
+// Copyright Epic Games, Inc. All Rights Reserved.
+
+#pragma once
+
+#include "MLDeformer.h"
+#include "MLDeformerAsset.h"
+#include "OptimusComputeDataInterface.h"
+#include "ComputeFramework/ComputeDataProvider.h"
+#include "MLDeformerDebugDataInterface.generated.h"
+
+class FRDGBuffer;
+class FRDGBufferSRV;
+struct FMLDeformerMeshMapping;
+class FRHIShaderResourceView;
+class FSkeletalMeshObject;
+class UMLDeformerComponent;
+class USkeletalMeshComponent;
+
+/** 
+ * Compute Framework Data Interface for MLDefomer debugging data. 
+ * This interfaces to editor only data, and so will only give valid results in that context.
+ */
+UCLASS(Category = ComputeFramework)
+class MLDEFORMER_API UMLDeformerDebugDataInterface : public UOptimusComputeDataInterface
+{
+	GENERATED_BODY()
+
+public:
+	//~ Begin UOptimusComputeDataInterface Interface
+	FString GetDisplayName() const override;
+	TArray<FOptimusCDIPinDefinition> GetPinDefinitions() const override;
+	//~ End UOptimusComputeDataInterface Interface
+
+	//~ Begin UComputeDataInterface Interface
+	void GetSupportedInputs(TArray<FShaderFunctionDefinition>& OutFunctions) const override;
+	void GetShaderParameters(TCHAR const* UID, FShaderParametersMetadataBuilder& OutBuilder) const override;
+	void GetHLSL(FString& OutHLSL) const override;
+	void GetSourceTypes(TArray<UClass*>& OutSourceTypes) const override;
+	UComputeDataProvider* CreateDataProvider(TArrayView<TObjectPtr<UObject>> InSourceObjects, uint64 InInputMask, uint64 InOutputMask) const override;
+	//~ End UComputeDataInterface Interface
+};
+
+
+/** Compute Framework Data Provider for MLDeformer debugging data. */
+UCLASS(BlueprintType, editinlinenew, Category = ComputeFramework)
+class MLDEFORMER_API UMLDeformerDebugDataProvider : public UComputeDataProvider
+{
+	GENERATED_BODY()
+
+public:
+	UPROPERTY(EditAnywhere, BlueprintReadWrite, Category = Binding)
+	TObjectPtr<USkeletalMeshComponent> SkeletalMeshComponent = nullptr;
+	UPROPERTY(EditAnywhere, BlueprintReadWrite, Category = Binding)
+	TObjectPtr<UMLDeformerComponent> DeformerComponent = nullptr;
+
+#if WITH_EDITORONLY_DATA
+	/* 
+	 * Mesh remapping data between the geometry cache and the rendered vertex buffer. 
+	 * This is calculated (slow) on construction. The data isn't available outside of the editor.
+	 */
+	TArray<FMLDeformerMeshMapping> MeshMappings;
+#endif
+
+	//~ Begin UComputeDataProvider Interface
+	bool IsValid() const override;
+	FComputeDataProviderRenderProxy* GetRenderProxy() override;
+	//~ End UComputeDataProvider Interface
+};
+
+
+/** Compute Framework Data Provider Proxy for MLDeformer debugging data. */
+class FMLDeformerDebugDataProviderProxy : public FComputeDataProviderRenderProxy
+{
+public:
+#if WITH_EDITORONLY_DATA
+	FMLDeformerDebugDataProviderProxy(USkeletalMeshComponent* SkeletalMeshComponent, UMLDeformerComponent* DeformerComponent, TArray<FMLDeformerMeshMapping> const& MeshMappings);
+#endif
+
+	//~ Begin FComputeDataProviderRenderProxy Interface
+	void AllocateResources(FRDGBuilder& GraphBuilder) override;
+	void GetBindings(int32 InvocationIndex, TCHAR const* UID, FBindings& OutBindings) const override;
+	//~ End FComputeDataProviderRenderProxy Interface
+
+private:
+	FSkeletalMeshObject* SkeletalMeshObject;
+	TArray<FVector3f> GroundTruthPositions;
+	FRHIShaderResourceView* VertexMapBufferSRV = nullptr;
+	FRDGBuffer* GroundTruthBuffer = nullptr;
+	FRDGBufferSRV* GroundTruthBufferSRV = nullptr;
+	float HeatMapScale = 1.0f;
+};
diff --git a/Engine/Plugins/Experimental/MLDeformer/Source/MLDeformer/Public/MLDeformerGraphDataInterface.h b/Engine/Plugins/Experimental/MLDeformer/Source/MLDeformer/Public/MLDeformerGraphDataInterface.h
deleted file mode 100644
diff --git a/Engine/Plugins/Experimental/MLDeformer/Source/MLDeformerEditor/MLDeformerEditor.Build.cs b/Engine/Plugins/Experimental/MLDeformer/Source/MLDeformerEditor/MLDeformerEditor.Build.cs
--- a/Engine/Plugins/Experimental/MLDeformer/Source/MLDeformerEditor/MLDeformerEditor.Build.cs
+++ b/Engine/Plugins/Experimental/MLDeformer/Source/MLDeformerEditor/MLDeformerEditor.Build.cs
@@ -46,6 +46,10 @@ namespace UnrealBuildTool.Rules
 					"ToolWidgets",
 					"ComputeFramework",
 					"OptimusDeveloper",
+
+					"OptimusCore",
+					"RenderCore",
+					"RHI",
 				}
 			);
 
