Non-Unity Fixes

diff --git a/Engine/Plugins/Experimental/MLDeformer/Source/MLDeformer/Private/MLDeformerDataInterface.cpp b/Engine/Plugins/Experimental/MLDeformer/Source/MLDeformer/Private/MLDeformerDataInterface.cpp
--- a/Engine/Plugins/Experimental/MLDeformer/Source/MLDeformer/Private/MLDeformerDataInterface.cpp
+++ b/Engine/Plugins/Experimental/MLDeformer/Source/MLDeformer/Private/MLDeformerDataInterface.cpp
@@ -13,6 +13,8 @@
 #include "RenderGraphResources.h"
 #include "ShaderParameterMetadataBuilder.h"
 #include "SkeletalRenderPublic.h"
+#include "Components/SkeletalMeshComponent.h"
+
 
 FString UMLDeformerDataInterface::GetDisplayName() const
 {
