MLDeformer: Use fixed debug ml deformer when showing heat maps.
In future might want to make this user configurable.
Add option to lerp rendered mesh between deformed and ground truth when in ground truth heatmap mode.
Route heatmap mode through ml debug data interface so that graph kernel can do the requested thing.
Pass reciprical of heat map scale to data interface instead of doing divide in kernel.

diff --git a/Engine/Plugins/Experimental/MLDeformer/Shaders/Private/MLDeformerDebugDataInterface.ush b/Engine/Plugins/Experimental/MLDeformer/Shaders/Private/MLDeformerDebugDataInterface.ush
--- a/Engine/Plugins/Experimental/MLDeformer/Shaders/Private/MLDeformerDebugDataInterface.ush
+++ b/Engine/Plugins/Experimental/MLDeformer/Shaders/Private/MLDeformerDebugDataInterface.ush
@@ -4,7 +4,9 @@
 
 uint DI_LOCAL(NumVertices);
 uint DI_LOCAL(InputStreamStart);
+uint DI_LOCAL(HeatMapMode);
 float DI_LOCAL(HeatMapScale);
+float DI_LOCAL(GroundTruthLerp);
 
 StructuredBuffer<float3> DI_LOCAL(PositionGroundTruthBuffer);
 Buffer<uint> DI_LOCAL(VertexMapBuffer);
@@ -20,7 +22,17 @@ DI_IMPL_READ(ReadPositionGroundTruth, float3, uint VertexIndex)
 	return DI_LOCAL(PositionGroundTruthBuffer)[BufferIndex];
 }
 
+DI_IMPL_READ(ReadHeatMapMode, uint, void)
+{
+	return DI_LOCAL(HeatMapMode);
+}
+
 DI_IMPL_READ(ReadHeatMapScale, float, void)
 {
 	return DI_LOCAL(HeatMapScale);
 }
+
+DI_IMPL_READ(ReadGroundTruthLerp, float, void)
+{
+	return DI_LOCAL(GroundTruthLerp);
+}
diff --git a/Engine/Plugins/Experimental/MLDeformer/Source/MLDeformer/Private/MLDeformerDebugDataInterface.cpp b/Engine/Plugins/Experimental/MLDeformer/Source/MLDeformer/Private/MLDeformerDebugDataInterface.cpp
--- a/Engine/Plugins/Experimental/MLDeformer/Source/MLDeformer/Private/MLDeformerDebugDataInterface.cpp
+++ b/Engine/Plugins/Experimental/MLDeformer/Source/MLDeformer/Private/MLDeformerDebugDataInterface.cpp
@@ -26,7 +26,9 @@ FString UMLDeformerDebugDataInterface::GetDisplayName() const
 TArray<FOptimusCDIPinDefinition> UMLDeformerDebugDataInterface::GetPinDefinitions() const
 {
 	TArray<FOptimusCDIPinDefinition> Defs;
+	Defs.Add({ "HeatMapMode", "ReadHeatMapMode" });
 	Defs.Add({ "HeatMapScale", "ReadHeatMapScale" });
+	Defs.Add({ "GroundTruthLerp", "ReadGroundTruthLerp" });
 	Defs.Add({ "PositionGroundTruth", "ReadPositionGroundTruth", Optimus::DomainName::Vertex, "ReadNumVertices" });
 	return Defs;
 }
@@ -42,6 +44,15 @@ void UMLDeformerDebugDataInterface::GetSupportedInputs(TArray<FShaderFunctionDef
 		Fn.ParamTypes.Add(ReturnParam);
 		OutFunctions.Add(Fn);
 	}
+	{
+		FShaderFunctionDefinition Fn;
+		Fn.Name = TEXT("ReadHeatMapMode");
+		Fn.bHasReturnType = true;
+		FShaderParamTypeDefinition ReturnParam = {};
+		ReturnParam.ValueType = FShaderValueType::Get(EShaderFundamentalType::Int);
+		Fn.ParamTypes.Add(ReturnParam);
+		OutFunctions.Add(Fn);
+	}
 	{
 		FShaderFunctionDefinition Fn;
 		Fn.Name = TEXT("ReadHeatMapScale");
@@ -51,6 +62,15 @@ void UMLDeformerDebugDataInterface::GetSupportedInputs(TArray<FShaderFunctionDef
 		Fn.ParamTypes.Add(ReturnParam);
 		OutFunctions.Add(Fn);
 	}
+	{
+		FShaderFunctionDefinition Fn;
+		Fn.Name = TEXT("ReadGroundTruthLerp");
+		Fn.bHasReturnType = true;
+		FShaderParamTypeDefinition ReturnParam = {};
+		ReturnParam.ValueType = FShaderValueType::Get(EShaderFundamentalType::Float);
+		Fn.ParamTypes.Add(ReturnParam);
+		OutFunctions.Add(Fn);
+	}
 	{
 		FShaderFunctionDefinition Fn;
 		Fn.Name = TEXT("ReadPositionGroundTruth");
@@ -68,7 +88,9 @@ void UMLDeformerDebugDataInterface::GetSupportedInputs(TArray<FShaderFunctionDef
 BEGIN_SHADER_PARAMETER_STRUCT(FMLDeformerDebugDataInterfaceParameters, )
 	SHADER_PARAMETER(uint32, NumVertices)
 	SHADER_PARAMETER(uint32, InputStreamStart)
+	SHADER_PARAMETER(int32, HeatMapMode)
 	SHADER_PARAMETER(float, HeatMapScale)
+	SHADER_PARAMETER(float, GroundTruthLerp)
 	SHADER_PARAMETER_RDG_BUFFER_SRV(StructuredBuffer<float3>, PositionGroundTruthBuffer)
 	SHADER_PARAMETER_SRV(Buffer<uint>, VertexMapBuffer)
 END_SHADER_PARAMETER_STRUCT()
@@ -201,7 +223,9 @@ FMLDeformerDebugDataProviderProxy::FMLDeformerDebugDataProviderProxy(USkeletalMe
 	
 	UMLDeformerAsset const* DeformerAsset = DeformerComponent->GetDeformerAsset();
 	VertexMapBufferSRV = DeformerAsset->GetVertexMapBuffer().ShaderResourceViewRHI;
-	HeatMapScale = DeformerAsset->GetVizSettings()->GetHeatMapScale();
+	HeatMapMode = (int32)DeformerAsset->GetVizSettings()->GetHeatMapMode();
+	HeatMapScale = 1.f / FMath::Max(DeformerAsset->GetVizSettings()->GetHeatMapScale(), 0.00001f);
+	GroundTruthLerp = DeformerAsset->GetVizSettings()->GetGroundTruthLerp();
 
 	const int32 LODIndex = 0;
 	const float SampleTime = SkeletalMeshComponent->GetPosition();
@@ -229,7 +253,9 @@ void FMLDeformerDebugDataProviderProxy::GetBindings(int32 InvocationIndex, TCHAR
 	FMemory::Memset(&Parameters, 0, sizeof(Parameters));
 	Parameters.NumVertices = 0;
 	Parameters.InputStreamStart = RenderSection.BaseVertexIndex;
+	Parameters.HeatMapMode = HeatMapMode;
 	Parameters.HeatMapScale = HeatMapScale;
+	Parameters.GroundTruthLerp = GroundTruthLerp;
 	Parameters.PositionGroundTruthBuffer = GroundTruthBufferSRV;
 	Parameters.VertexMapBuffer = VertexMapBufferSRV;
 
diff --git a/Engine/Plugins/Experimental/MLDeformer/Source/MLDeformer/Public/MLDeformerDebugDataInterface.h b/Engine/Plugins/Experimental/MLDeformer/Source/MLDeformer/Public/MLDeformerDebugDataInterface.h
--- a/Engine/Plugins/Experimental/MLDeformer/Source/MLDeformer/Public/MLDeformerDebugDataInterface.h
+++ b/Engine/Plugins/Experimental/MLDeformer/Source/MLDeformer/Public/MLDeformerDebugDataInterface.h
@@ -87,5 +87,7 @@ private:
 	FRHIShaderResourceView* VertexMapBufferSRV = nullptr;
 	FRDGBuffer* GroundTruthBuffer = nullptr;
 	FRDGBufferSRV* GroundTruthBufferSRV = nullptr;
-	float HeatMapScale = 1.0f;
+	int32 HeatMapMode = 0;
+	float HeatMapScale = 0.0f;
+	float GroundTruthLerp = 0.0f;
 };
diff --git a/Engine/Plugins/Experimental/MLDeformer/Source/MLDeformer/Public/MLDeformerVizSettings.h b/Engine/Plugins/Experimental/MLDeformer/Source/MLDeformer/Public/MLDeformerVizSettings.h
--- a/Engine/Plugins/Experimental/MLDeformer/Source/MLDeformer/Public/MLDeformerVizSettings.h
+++ b/Engine/Plugins/Experimental/MLDeformer/Source/MLDeformer/Public/MLDeformerVizSettings.h
@@ -26,11 +26,10 @@ enum class EMLDeformerVizMode : uint8
 UENUM()
 enum class EMLDeformerHeatMapMode : uint8
 {
-	// Visualize the error versus the ground truth model. Requires a ground truth model to be setup.
-	GroundTruth = 0,
-
 	// Visualize areas where the deformer is applying corrections.
-	Activations
+	Activations = 0,
+	// Visualize the error versus the ground truth model. Requires a ground truth model to be setup.
+	GroundTruth
 };
 
 /**
@@ -69,6 +68,7 @@ public:
 	EMLDeformerVizMode GetTempVisualizationMode() const { return TempVisualizationMode; }
 	EMLDeformerHeatMapMode GetHeatMapMode() const { return HeatMapMode; }
 	float GetHeatMapScale() const { return HeatMapScale; }
+	float GetGroundTruthLerp() const { return GroundTruthLerp; }
 	
 	void SetTempVisualizationMode(EMLDeformerVizMode Mode) { TempVisualizationMode = Mode; }
 	void SetDeformerGraph(UComputeGraph* InDeformerGraph) { DeformerGraph = InDeformerGraph; }
@@ -107,12 +107,16 @@ public:
 
 	/** What should the heatmap visualize? */
 	UPROPERTY(EditAnywhere, Category = "Live Settings", meta = (EditCondition = "bShowHeatMap"))
-	EMLDeformerHeatMapMode HeatMapMode = EMLDeformerHeatMapMode::GroundTruth;
+	EMLDeformerHeatMapMode HeatMapMode = EMLDeformerHeatMapMode::Activations;
 
 	/** How many units (centimeters) of error should the most intense color represent? */
 	UPROPERTY(EditAnywhere, Category = "Live Settings", meta = (EditCondition = "bShowHeatMap", ClampMin = "0.00001"))
 	float HeatMapScale = 1.0f;
 
+	/** Lerp from ML deformed model to ground truth model when in heat map mode. */
+	UPROPERTY(EditAnywhere, Category = "Live Settings", meta = (EditCondition = "HeatMapMode==EMLDeformerHeatMapMode::GroundTruth", UIMin = "0.0", UIMax = "1.0", ClampMin = "0.0", ClampMax = "1.0"))
+	float GroundTruthLerp = 0.0f;
+
 	/** Draw the linear skinned actor? */
 	UPROPERTY(EditAnywhere, Category = "Live Settings")
 	bool bDrawLinearSkinnedActor = true;
diff --git a/Engine/Plugins/Experimental/MLDeformer/Source/MLDeformerEditor/Private/MLDeformerEditorData.cpp b/Engine/Plugins/Experimental/MLDeformer/Source/MLDeformerEditor/Private/MLDeformerEditorData.cpp
--- a/Engine/Plugins/Experimental/MLDeformer/Source/MLDeformerEditor/Private/MLDeformerEditorData.cpp
+++ b/Engine/Plugins/Experimental/MLDeformer/Source/MLDeformerEditor/Private/MLDeformerEditorData.cpp
@@ -215,6 +215,10 @@ void FMLDeformerEditorData::UpdateDeformerGraph()
 		UMLDeformerVizSettings* VizSettings = MLDeformerAsset->GetVizSettings();
 		check(VizSettings);
 		UComputeGraph* ComputeGraph = MLDeformerAsset->GetInferenceNeuralNetwork() ? VizSettings->GetDeformerGraph() : nullptr;
+		
+		const bool bUseHeatMapDeformer = VizSettings->GetShowHeatMap();
+		ComputeGraph = bUseHeatMapDeformer ? HeatMapDeformerGraph.Get() : ComputeGraph;
+		
 		UOptimusMeshDeformer* MeshDeformer = NewObject<UOptimusMeshDeformer>(SkelMeshComponent);
 		MeshDeformer->ComputeGraph = ComputeGraph;
 		SkelMeshComponent->SetMeshDeformer(MeshDeformer);
@@ -517,13 +521,19 @@ FText FMLDeformerEditorData::GetOverlayText()
 	return FText::GetEmpty();
 }
 
-void FMLDeformerEditorData::CreateMaterials()
+void FMLDeformerEditorData::CreateHeatMapAssets()
 {
-	// Try load the heat map material.
-	const FString HeatMapPath = TEXT("/MLDeformer/Materials/HeatMap.HeatMap");
-	UObject* Object = StaticLoadObject(UMaterial::StaticClass(), nullptr, *HeatMapPath);
-	HeatMapMaterial = Cast<UMaterial>(Object);
+	// Could be better to explicitly expose in UI. More flexibility and no sync load required here?
+
+	const FString HeatMapMaterialPath = TEXT("/MLDeformer/Materials/HeatMap.HeatMap");
+	UObject* MaterialObject = StaticLoadObject(UMaterial::StaticClass(), nullptr, *HeatMapMaterialPath);
+	HeatMapMaterial = Cast<UMaterial>(MaterialObject);
 	check(HeatMapMaterial);
+
+	const FString HeatMapDeformerPath = TEXT("/MLDeformer/Deformers/DebugMLDeformerGraph.DebugMLDeformerGraph");
+	UObject* DeformerObject = StaticLoadObject(UComputeGraph::StaticClass(), nullptr, *HeatMapDeformerPath);
+	HeatMapDeformerGraph = Cast<UComputeGraph>(DeformerObject);
+	check(HeatMapDeformerGraph);
 }
 
 void FMLDeformerEditorData::SetHeatMapMaterialEnabled(bool bEnabled)
diff --git a/Engine/Plugins/Experimental/MLDeformer/Source/MLDeformerEditor/Private/MLDeformerEditorToolkit.cpp b/Engine/Plugins/Experimental/MLDeformer/Source/MLDeformerEditor/Private/MLDeformerEditorToolkit.cpp
--- a/Engine/Plugins/Experimental/MLDeformer/Source/MLDeformerEditor/Private/MLDeformerEditorToolkit.cpp
+++ b/Engine/Plugins/Experimental/MLDeformer/Source/MLDeformerEditor/Private/MLDeformerEditorToolkit.cpp
@@ -657,8 +657,8 @@ void FMLDeformerEditorToolkit::HandlePreviewSceneCreated(const TSharedRef<IPerso
 	// Start playing the animations.
 	EditorData->InitAssets();
 
-	// Create visualization materials.
-	EditorData->CreateMaterials();
+	// Create visualization assets.
+	EditorData->CreateHeatMapAssets();
 	EditorData->SetHeatMapMaterialEnabled(EditorData->GetDeformerAsset()->GetVizSettings()->GetShowHeatMap());
 
 	VizSettings->SetTempVisualizationMode(VizSettings->GetVisualizationMode());
@@ -739,6 +739,7 @@ void FMLDeformerEditorToolkit::OnFinishedChangingDetails(const FPropertyChangedE
 	if (Property->GetFName() == GET_MEMBER_NAME_CHECKED(UMLDeformerVizSettings, bShowHeatMap))
 	{
 		EditorData->SetHeatMapMaterialEnabled(EditorData->GetDeformerAsset()->GetVizSettings()->GetShowHeatMap());
+		EditorData->UpdateDeformerGraph();
 	}
 	else
 	if (Property->GetFName() == GET_MEMBER_NAME_CHECKED(UMLDeformerVizSettings, VisualizationMode))
diff --git a/Engine/Plugins/Experimental/MLDeformer/Source/MLDeformerEditor/Private/MLDeformerVizSettingsDetails.cpp b/Engine/Plugins/Experimental/MLDeformer/Source/MLDeformerEditor/Private/MLDeformerVizSettingsDetails.cpp
--- a/Engine/Plugins/Experimental/MLDeformer/Source/MLDeformerEditor/Private/MLDeformerVizSettingsDetails.cpp
+++ b/Engine/Plugins/Experimental/MLDeformer/Source/MLDeformerEditor/Private/MLDeformerVizSettingsDetails.cpp
@@ -220,6 +220,7 @@ void FMLDeformerVizSettingsDetails::CustomizeDetails(class IDetailLayoutBuilder&
 	HeatMapGroup.AddPropertyRow(DetailBuilder.GetProperty(GET_MEMBER_NAME_CHECKED(UMLDeformerVizSettings, bShowHeatMap)));
 	HeatMapGroup.AddPropertyRow(DetailBuilder.GetProperty(GET_MEMBER_NAME_CHECKED(UMLDeformerVizSettings, HeatMapMode)));
 	HeatMapGroup.AddPropertyRow(DetailBuilder.GetProperty(GET_MEMBER_NAME_CHECKED(UMLDeformerVizSettings, HeatMapScale)));
+	HeatMapGroup.AddPropertyRow(DetailBuilder.GetProperty(GET_MEMBER_NAME_CHECKED(UMLDeformerVizSettings, GroundTruthLerp)));
 
 	IDetailGroup& VisGroup = LiveSettingsCategory.AddGroup("Visibility", LOCTEXT("Visibility", "Visibility"), false, true);
 	VisGroup.AddPropertyRow(DetailBuilder.GetProperty(GET_MEMBER_NAME_CHECKED(UMLDeformerVizSettings, bDrawLinearSkinnedActor)));
diff --git a/Engine/Plugins/Experimental/MLDeformer/Source/MLDeformerEditor/Public/MLDeformerEditorData.h b/Engine/Plugins/Experimental/MLDeformer/Source/MLDeformerEditor/Public/MLDeformerEditorData.h
--- a/Engine/Plugins/Experimental/MLDeformer/Source/MLDeformerEditor/Public/MLDeformerEditorData.h
+++ b/Engine/Plugins/Experimental/MLDeformer/Source/MLDeformerEditor/Public/MLDeformerEditorData.h
@@ -110,7 +110,7 @@ public:
 	static FString GetDefaultDeformerGraphAssetPath();
 	static UComputeGraph* LoadDefaultDeformerGraph();
 	void SetAnimFrame(int32 FrameNumber);
-	void CreateMaterials();
+	void CreateHeatMapAssets();
 	void ClampFrameIndex();
 
 	void UpdateIsReadyForTrainingState();
@@ -160,6 +160,9 @@ private:
 	/** The heatmap material. */
 	TObjectPtr<UMaterial> HeatMapMaterial;
 
+	/** The heatmap material. */
+	TObjectPtr<UComputeGraph> HeatMapDeformerGraph;
+
 	/** The world that our actors are inside. */
 	TObjectPtr<UWorld> World;
 
