MLDeformer: Use selected geometry cache for ground truth heat map.
But only if it matches current test animation.

diff --git a/Engine/Plugins/Experimental/MLDeformer/Source/MLDeformer/Private/MLDeformerDebugDataInterface.cpp b/Engine/Plugins/Experimental/MLDeformer/Source/MLDeformer/Private/MLDeformerDebugDataInterface.cpp
--- a/Engine/Plugins/Experimental/MLDeformer/Source/MLDeformer/Private/MLDeformerDebugDataInterface.cpp
+++ b/Engine/Plugins/Experimental/MLDeformer/Source/MLDeformer/Private/MLDeformerDebugDataInterface.cpp
@@ -155,33 +155,51 @@ FComputeDataProviderRenderProxy* UMLDeformerDebugDataProvider::GetRenderProxy()
 
 #if WITH_EDITORONLY_DATA
 
+// Return a GeometryCache only if it matches the current test sequence.
+UGeometryCache const* GetActiveGeometryCache(UMLDeformerAsset const* DeformerAsset)
+{
+	UGeometryCache const* GeometryCache = DeformerAsset->GetVizSettings()->GetGroundTruth();
+	GeometryCache = GeometryCache != nullptr ? GeometryCache : DeformerAsset->GetGeometryCache();
+
+	UAnimSequence const* AnimSequence = DeformerAsset->GetVizSettings()->GetTestAnimSequence();
+
+	if (GeometryCache != nullptr && AnimSequence != nullptr)
+	{
+		const float AnimSeqDuration = AnimSequence->GetPlayLength();
+		const float GeomCacheDuration = GeometryCache->CalculateDuration();
+		if (FMath::Abs(AnimSeqDuration - GeomCacheDuration) < 0.001f)
+		{
+			return GeometryCache;
+		}
+	}
+
+	return nullptr;
+}
+
+// Fill the ground truth positions from the geometry cache.
 void GetGroundTruthPositions(
 	int32 LODIndex,
 	float SampleTime,
 	UMLDeformerAsset const* DeformerAsset,
 	TArray<FMLDeformerMeshMapping> const& MeshMappings,
+	UGeometryCache const* GeometryCache,
 	TArray<FVector3f>& OutPositions)
 {
-	if (!ensure(DeformerAsset))
+	if (GeometryCache == nullptr)
 	{
 		return;
 	}
 	USkeletalMesh const* SkelMesh = DeformerAsset->GetSkeletalMesh();
-	if (!ensure(SkelMesh))
+	if (!ensure(SkelMesh != nullptr))
 	{
 		return;
 	}
 	FSkeletalMeshModel const* ImportedModel = SkelMesh->GetImportedModel();
-	if (!ensure(ImportedModel))
+	if (!ensure(ImportedModel != nullptr))
 	{
 		return;
 	}
-	UGeometryCache const* GeomCache = DeformerAsset->GetGeometryCache();
-	if (!ensure(GeomCache))
-	{
-		return;
-	}
-
+	
 	const FTransform AlignmentTransform = DeformerAsset->GetAlignmentTransform();
 
 	FSkeletalMeshLODModel const& LODModel = ImportedModel->LODModels[LODIndex];
@@ -196,7 +214,7 @@ void GetGroundTruthPositions(
 	{
 		FMLDeformerMeshMapping const& MeshMapping = MeshMappings[MeshMappingIndex];
 		FSkelMeshImportedMeshInfo const& MeshInfo = SkelMeshInfos[MeshMapping.MeshIndex];
-		UGeometryCacheTrack* Track = GeomCache->Tracks[MeshMapping.TrackIndex];
+		UGeometryCacheTrack* Track = GeometryCache->Tracks[MeshMapping.TrackIndex];
 
 		FGeometryCacheMeshData GeomCacheMeshData;
 		if (!Track->GetMeshDataAtTime(SampleTime, GeomCacheMeshData))
@@ -229,7 +247,22 @@ FMLDeformerDebugDataProviderProxy::FMLDeformerDebugDataProviderProxy(USkeletalMe
 
 	const int32 LODIndex = 0;
 	const float SampleTime = SkeletalMeshComponent->GetPosition();
-	GetGroundTruthPositions(0, SampleTime, DeformerAsset, MeshMappings, GroundTruthPositions);
+	UGeometryCache const* GroundTruthGeomCache = GetActiveGeometryCache(DeformerAsset);
+	GetGroundTruthPositions(0, SampleTime, DeformerAsset, MeshMappings, GroundTruthGeomCache, GroundTruthPositions);
+	
+	if (GroundTruthPositions.Num() == 0)
+	{	
+		// We didn't get valid ground truth vertices.
+		// Make non empty array for later buffer generation.
+		GroundTruthPositions.Add(FVector3f::ZeroVector);
+		// Silently disable relevant debug things.
+		if (HeatMapMode == (int32)EMLDeformerHeatMapMode::GroundTruth)
+		{
+			HeatMapMode = -1;
+			HeatMapScale = 0.f;
+			GroundTruthLerp = 0.f;
+		}
+	}
 }
 
 #endif // WITH_EDITORONLY_DATA
@@ -239,7 +272,7 @@ void FMLDeformerDebugDataProviderProxy::AllocateResources(FRDGBuilder& GraphBuil
 	GroundTruthBuffer = GraphBuilder.CreateBuffer(FRDGBufferDesc::CreateStructuredDesc(sizeof(FVector3f), GroundTruthPositions.Num()), TEXT("MLDeformer.GroundTruthPositions"));
 	GroundTruthBufferSRV = GraphBuilder.CreateSRV(GroundTruthBuffer);
 
-	GraphBuilder.QueueBufferUpload(GroundTruthBuffer, GroundTruthPositions.GetData(), GroundTruthPositions.GetAllocatedSize(), ERDGInitialDataFlags::None);
+	GraphBuilder.QueueBufferUpload(GroundTruthBuffer, GroundTruthPositions.GetData(), sizeof(FVector3f) * GroundTruthPositions.Num(), ERDGInitialDataFlags::None);
 }
 
 void FMLDeformerDebugDataProviderProxy::GetBindings(int32 InvocationIndex, TCHAR const* UID, FBindings& OutBindings) const
diff --git a/Engine/Plugins/Experimental/MLDeformer/Source/MLDeformer/Public/MLDeformerVizSettings.h b/Engine/Plugins/Experimental/MLDeformer/Source/MLDeformer/Public/MLDeformerVizSettings.h
--- a/Engine/Plugins/Experimental/MLDeformer/Source/MLDeformer/Public/MLDeformerVizSettings.h
+++ b/Engine/Plugins/Experimental/MLDeformer/Source/MLDeformer/Public/MLDeformerVizSettings.h
@@ -68,7 +68,7 @@ public:
 	EMLDeformerVizMode GetTempVisualizationMode() const { return TempVisualizationMode; }
 	EMLDeformerHeatMapMode GetHeatMapMode() const { return HeatMapMode; }
 	float GetHeatMapScale() const { return HeatMapScale; }
-	float GetGroundTruthLerp() const { return GroundTruthLerp; }
+	float GetGroundTruthLerp() const { return HeatMapMode == EMLDeformerHeatMapMode::GroundTruth ? GroundTruthLerp : 0.f; }
 	
 	void SetTempVisualizationMode(EMLDeformerVizMode Mode) { TempVisualizationMode = Mode; }
 	void SetDeformerGraph(UComputeGraph* InDeformerGraph) { DeformerGraph = InDeformerGraph; }
