[MLDeformer] Added a bone include list with a button to automatically add animated bones.

diff --git a/Engine/Plugins/Experimental/MLDeformer/Source/MLDeformer/Private/MLDeformerAsset.cpp b/Engine/Plugins/Experimental/MLDeformer/Source/MLDeformer/Private/MLDeformerAsset.cpp
--- a/Engine/Plugins/Experimental/MLDeformer/Source/MLDeformer/Private/MLDeformerAsset.cpp
+++ b/Engine/Plugins/Experimental/MLDeformer/Source/MLDeformer/Private/MLDeformerAsset.cpp
@@ -12,6 +12,7 @@
 #include "NeuralNetwork.h"
 #include "Engine/SkeletalMesh.h"
 #include "Animation/AnimSequence.h"
+#include "Animation/AnimData/AnimDataModel.h"
 
 #define LOCTEXT_NAMESPACE "MLDeformerAsset"
 
@@ -86,6 +87,19 @@ void UMLDeformerAsset::PostLoad()
 	InitGPUData();
 }
 
+// Used for the FBoenReference, so it knows what skeleton to pick bones from.
+USkeleton* UMLDeformerAsset::GetSkeleton(bool& bInvalidSkeletonIsError)
+{
+	bInvalidSkeletonIsError = false;
+#if WITH_EDITOR
+	if (SkeletalMesh)
+	{
+		return SkeletalMesh->GetSkeleton();
+	}
+#endif
+	return nullptr;
+}
+
 #if WITH_EDITOR
 // A fuzzy name match.
 // There is a match when the track name starts with the mesh name.
@@ -308,6 +322,93 @@ void UMLDeformerAsset::InitVertexMap()
 	}
 }
 
+void UMLDeformerAsset::InitBoneIncludeListToAnimatedBonesOnly()
+{
+	if (!AnimSequence)
+	{
+		UE_LOG(LogMLDeformer, Warning, TEXT("Cannot initialize bone list as no Anim Sequence has been picked."));
+		return;
+	}
+
+	const UAnimDataModel* DataModel = AnimSequence->GetDataModel();
+	if (!DataModel)
+	{
+		UE_LOG(LogMLDeformer, Warning, TEXT("Anim sequence has no data model."));
+		return;
+	}
+
+	if (!SkeletalMesh)
+	{
+		UE_LOG(LogMLDeformer, Warning, TEXT("Skeletal Mesh has not been set."));
+		return;
+	}
+
+	USkeleton* Skeleton = SkeletalMesh->GetSkeleton();
+	if (!Skeleton)
+	{
+		UE_LOG(LogMLDeformer, Warning, TEXT("Skeletal Mesh has no skeleton."));
+		return;
+	}
+
+	// Iterate over all bones that are both in the skeleton and the animation.
+	TArray<FName> AnimatedBoneList;
+	const FReferenceSkeleton& RefSkeleton = Skeleton->GetReferenceSkeleton();
+	const int32 NumBones = RefSkeleton.GetNum();
+	for (int32 Index = 0; Index < NumBones; ++Index)
+	{
+		const FName BoneName = RefSkeleton.GetBoneName(Index);
+		const int32 BoneTrackIndex = DataModel->GetBoneTrackIndexByName(BoneName);
+		if (BoneTrackIndex == INDEX_NONE)
+		{
+			continue;
+		}
+
+		// Check if there is actually animation data.
+		const FBoneAnimationTrack& BoneAnimTrack = DataModel->GetBoneTrackByIndex(BoneTrackIndex);
+		const TArray<FQuat4f>& Rotations = BoneAnimTrack.InternalTrackData.RotKeys;
+		bool bIsAnimated = false;
+		if (!Rotations.IsEmpty())
+		{
+			const FQuat4f FirstQuat = Rotations[0];
+			for (const FQuat4f KeyValue : Rotations)
+			{
+				if (!KeyValue.Equals(FirstQuat))
+				{
+					bIsAnimated = true;
+					break;
+				}
+			}
+
+			if (!bIsAnimated)
+			{
+				UE_LOG(LogMLDeformer, Display, TEXT("Bone '%s' has keyframes but isn't animated."), *BoneName.ToString());
+			}
+		}
+
+		if (bIsAnimated)
+		{
+			AnimatedBoneList.Add(BoneName);
+		}
+	}
+
+	// Init the bone include list using the animated bones.
+	if (!AnimatedBoneList.IsEmpty())
+	{
+		BoneIncludeList.Empty();
+		BoneIncludeList.Reserve(AnimatedBoneList.Num());
+		for (FName BoneName : AnimatedBoneList)
+		{
+			BoneIncludeList.AddDefaulted();
+			FBoneReference& BoneRef = BoneIncludeList.Last();
+			BoneRef.BoneName = BoneName;
+		}
+	}
+	else
+	{
+		UE_LOG(LogMLDeformer, Warning, TEXT("There are no animated bone rotations in Anim Sequence '%s'."), *AnimSequence->GetName());
+	}
+}
+
 int32 UMLDeformerAsset::GetNumFramesForTraining() const 
 { 
 	return FMath::Min(GetNumFrames(), GetTrainingFrameLimit()); 
@@ -492,8 +593,31 @@ FMLDeformerInputInfo UMLDeformerAsset::CreateInputInfo() const
 	Settings.TargetMesh = GeometryCache;
 	Settings.bIncludeBones = (TrainingInputs == ETrainingInputs::BonesAndCurves || TrainingInputs == ETrainingInputs::BonesOnly);
 	Settings.bIncludeCurves = (TrainingInputs == ETrainingInputs::BonesAndCurves || TrainingInputs == ETrainingInputs::CurvesOnly);
-	Result.Init(Settings);
 
+	// Set the list of bones to use, from the bone references.
+	if (!BoneIncludeList.IsEmpty())
+	{
+		FString BoneName;
+		for (const FBoneReference& BoneReference : BoneIncludeList)
+		{
+			if (SkeletalMesh && SkeletalMesh->GetSkeleton() && BoneReference.BoneName.IsValid())
+			{
+				BoneName = BoneReference.BoneName.ToString();
+				const int32 BoneId = SkeletalMesh->GetRefSkeleton().FindRawBoneIndex(BoneReference.BoneName);
+				if (BoneId != INDEX_NONE && !Settings.BoneNamesToInclude.Contains(BoneName))
+				{
+					Settings.BoneNamesToInclude.Add(BoneName);
+					UE_LOG(LogMLDeformer, Display, TEXT("Including bone '%s' in training."), *BoneName);
+				}
+			}
+		}
+	}
+	else
+	{
+		UE_LOG(LogMLDeformer, Display, TEXT("Including ALL bones of skeleton in training."));
+	}
+
+	Result.Init(Settings);
 	return Result;
 }
 
diff --git a/Engine/Plugins/Experimental/MLDeformer/Source/MLDeformer/Public/MLDeformerAsset.h b/Engine/Plugins/Experimental/MLDeformer/Source/MLDeformer/Public/MLDeformerAsset.h
--- a/Engine/Plugins/Experimental/MLDeformer/Source/MLDeformer/Public/MLDeformerAsset.h
+++ b/Engine/Plugins/Experimental/MLDeformer/Source/MLDeformer/Public/MLDeformerAsset.h
@@ -6,9 +6,11 @@
 #include "UObject/NoExportTypes.h"
 #include "MLDeformerVizSettings.h"
 #include "MLDeformerInputInfo.h"
+#include "BoneContainer.h"
 #include "RenderCommandFence.h"
 #include "UObject/Object.h"
 #include "RenderResource.h"
+#include "Interfaces/Interface_BoneReferenceSkeletonProvider.h"
 #include "MLDeformerAsset.generated.h"
 
 // Forward declarations.
@@ -17,6 +19,7 @@ class USkeletalMesh;
 class UGeometryCache;
 class UAnimSequence;
 class UMLDeformerAsset;
+class USkeleton;
 
 /** The activation function to use during the ML Deformer training process. */
 UENUM()
@@ -104,7 +107,9 @@ struct MLDEFORMER_API FMLDeformerMeshMapping
  * In the editor it contains the skeletal mesh and geometry cache that are required to calculate vertex position deltas.
  */
 UCLASS(BlueprintType, hidecategories=Object)
-class MLDEFORMER_API UMLDeformerAsset : public UObject
+class MLDEFORMER_API UMLDeformerAsset 
+	: public UObject
+	, public IBoneReferenceSkeletonProvider
 {
 	GENERATED_BODY()
 
@@ -121,6 +126,12 @@ public:
 #endif
 	//~End of UObject interface
 
+
+	//~IBoneReferenceSkeletonProvider interface
+	USkeleton* GetSkeleton(bool& bInvalidSkeletonIsError) override;
+	//~End of IBoneReferenceSkeletonProvider interface
+
+
 #if WITH_EDITOR
 	FText GetGeomCacheErrorText(UGeometryCache* InGeomCache) const;
 	FText GetAnimSequenceErrorText(UGeometryCache* InGeomCache, UAnimSequence* InAnimSequence) const;
@@ -139,6 +150,8 @@ public:
 	void UpdateCachedNumVertices();
 	bool IsCompatibleWithNeuralNet() const;
 	FMLDeformerInputInfo CreateInputInfo() const;
+
+	void InitBoneIncludeListToAnimatedBonesOnly();
 #endif
 
 	void InitVertexMap();
@@ -171,6 +184,9 @@ public:
 	const UMLDeformerVizSettings* GetVizSettings() const { return VizSettings; }
 	UMLDeformerVizSettings* GetVizSettings() { return VizSettings; }
 
+	TArray<FBoneReference>& GetBoneIncludeList() { return BoneIncludeList; }
+	const TArray<FBoneReference>& GetBoneIncludeList() const { return BoneIncludeList; }
+
 	int32 GetNumFrames() const;
 	int32 GetTrainingFrameLimit() const { return MaxTrainingFrames; }
 	int32 GetNumFramesForTraining() const;
@@ -313,5 +329,9 @@ public:
 	  * Deltas that are longer than the cutoff value (in units), will be ignored and set to zero length. */
 	UPROPERTY(EditAnywhere, BlueprintReadWrite, Category = "Delta Generation", meta = (ClampMin = "0.01"))
 	float DeltaCutoffLength = 30.0f;
+
+	/** The bones to include during training. When none are provided, all bones of the Skeleton will be included. */
+	UPROPERTY(EditAnywhere, Category = "Inputs and Output", meta = (EditCondition = "TrainingInputs == ETrainingInputs::BonesAndCurves || TrainingInputs == ETrainingInputs::BonesOnly"))
+	TArray<FBoneReference> BoneIncludeList;
 #endif // WITH_EDITORONLY_DATA
 };
diff --git a/Engine/Plugins/Experimental/MLDeformer/Source/MLDeformerEditor/Private/MLDeformerAssetDetails.cpp b/Engine/Plugins/Experimental/MLDeformer/Source/MLDeformerEditor/Private/MLDeformerAssetDetails.cpp
--- a/Engine/Plugins/Experimental/MLDeformer/Source/MLDeformerEditor/Private/MLDeformerAssetDetails.cpp
+++ b/Engine/Plugins/Experimental/MLDeformer/Source/MLDeformerEditor/Private/MLDeformerAssetDetails.cpp
@@ -16,12 +16,14 @@
 #include "DetailCategoryBuilder.h"
 #include "DetailWidgetRow.h"
 #include "IDetailsView.h"
+#include "IDetailGroup.h"
 
 #include "Widgets/DeclarativeSyntaxSupport.h"
 #include "Widgets/Layout/SBox.h"
 #include "DetailLayoutBuilder.h"
 
 #include "SWarningOrErrorBox.h"
+#include "Widgets/Input/SButton.h"
 
 #define LOCTEXT_NAMESPACE "MLDeformerAssetDetails"
 
@@ -236,6 +238,22 @@ void FMLDeformerAssetDetails::CustomizeDetails(class IDetailLayoutBuilder& Detai
 
 	InputOutputCategoryBuilder.AddProperty(GET_MEMBER_NAME_CHECKED(UMLDeformerAsset, DeltaCutoffLength));
 
+	// Bone include list group.
+	IDetailGroup& BoneIncludeGroup = InputOutputCategoryBuilder.AddGroup("BoneIncludeGroup", LOCTEXT("BoneIncludeGroup", "Bones"), false, true);
+	BoneIncludeGroup.AddWidgetRow()
+		//.WholeRowContent()
+		.ValueContent()
+		//.NameContent()
+		[
+			SNew(SButton)
+			.HAlign(HAlign_Center)
+			.VAlign(VAlign_Center)
+			.Text(LOCTEXT("AnimatedBonesButton", "Animated Bones Only"))
+			.OnClicked(FOnClicked::CreateSP(this, &FMLDeformerAssetDetails::OnFilterAnimatedBonesOnly, DeformerAsset))
+			.IsEnabled_Lambda([DeformerAsset](){ return (DeformerAsset->GetTrainingInputs() == ETrainingInputs::BonesAndCurves) || (DeformerAsset->GetTrainingInputs() == ETrainingInputs::BonesOnly); })
+		];
+	BoneIncludeGroup.AddPropertyRow(DetailBuilder.GetProperty(GET_MEMBER_NAME_CHECKED(UMLDeformerAsset, BoneIncludeList)));
+
 	// Training settings.
 	IDetailCategoryBuilder& SettingsCategoryBuilder = DetailBuilder.EditCategory("Training Settings", FText::GetEmpty(), ECategoryPriority::Important);
 	SettingsCategoryBuilder.AddProperty(GET_MEMBER_NAME_CHECKED(UMLDeformerAsset, NumHiddenLayers));
@@ -302,4 +320,11 @@ bool FMLDeformerAssetDetails::FilterAnimSequences(const FAssetData& AssetData, U
 	return true;
 }
 
+FReply FMLDeformerAssetDetails::OnFilterAnimatedBonesOnly(UMLDeformerAsset* DeformerAsset) const
+{
+	DeformerAsset->InitBoneIncludeListToAnimatedBonesOnly();
+	DetailLayoutBuilder->ForceRefreshDetails();
+	return FReply::Handled();
+}
+
 #undef LOCTEXT_NAMESPACE
diff --git a/Engine/Plugins/Experimental/MLDeformer/Source/MLDeformerEditor/Private/MLDeformerAssetDetails.h b/Engine/Plugins/Experimental/MLDeformer/Source/MLDeformerEditor/Private/MLDeformerAssetDetails.h
--- a/Engine/Plugins/Experimental/MLDeformer/Source/MLDeformerEditor/Private/MLDeformerAssetDetails.h
+++ b/Engine/Plugins/Experimental/MLDeformer/Source/MLDeformerEditor/Private/MLDeformerAssetDetails.h
@@ -5,10 +5,11 @@
 #include "CoreMinimal.h"
 #include "IDetailCustomization.h"
 #include "AssetRegistry/AssetData.h"
+#include "Input/Reply.h"
 
-class FMLDeformerEditorData;
 class IDetailLayoutBuilder;
 class USkeleton;
+class UMLDeformerAsset;
 
 class FMLDeformerAssetDetails : public IDetailCustomization
 {
@@ -23,14 +24,12 @@ public:
 
 private:
 	bool FilterAnimSequences(const FAssetData& AssetData, USkeleton* Skeleton);
+	FReply OnFilterAnimatedBonesOnly(UMLDeformerAsset* DeformerAsset) const;
 
 private:
 	/** Associated detail layout builder. */
 	IDetailLayoutBuilder* DetailLayoutBuilder = nullptr;
 
-	/** A pointer to the editor data. */
-	FMLDeformerEditorData* EditorData = nullptr;
-
 	/** The error text to show in the target mesh category. Empty if no error. */
 	FText TargetMeshErrorText;
 };
diff --git a/Engine/Plugins/Experimental/MLDeformer/Source/MLDeformerEditor/Private/MLDeformerEditorToolkit.cpp b/Engine/Plugins/Experimental/MLDeformer/Source/MLDeformerEditor/Private/MLDeformerEditorToolkit.cpp
--- a/Engine/Plugins/Experimental/MLDeformer/Source/MLDeformerEditor/Private/MLDeformerEditorToolkit.cpp
+++ b/Engine/Plugins/Experimental/MLDeformer/Source/MLDeformerEditor/Private/MLDeformerEditorToolkit.cpp
@@ -150,7 +150,6 @@ void FMLDeformerEditorToolkit::FillToolbar(FToolBarBuilder& ToolbarBuilder)
 				FExecuteAction::CreateLambda([this]()
 				{
 					UMLDeformerAsset* DeformerAsset = EditorData->GetDeformerAsset();
-					DeformerAsset->SetInputInfo(DeformerAsset->CreateInputInfo());
 
 					// Ask if we want to retrain the network if we already have something trained.
 					if (DeformerAsset->GetInferenceNeuralNetwork() != nullptr)
