[MLDeformer] Fixed a compile issue.

diff --git a/Engine/Plugins/Experimental/MLDeformer/Source/MLDeformer/Private/MLDeformerAsset.cpp b/Engine/Plugins/Experimental/MLDeformer/Source/MLDeformer/Private/MLDeformerAsset.cpp
--- a/Engine/Plugins/Experimental/MLDeformer/Source/MLDeformer/Private/MLDeformerAsset.cpp
+++ b/Engine/Plugins/Experimental/MLDeformer/Source/MLDeformer/Private/MLDeformerAsset.cpp
@@ -88,7 +88,7 @@ void UMLDeformerAsset::PostLoad()
 }
 
 // Used for the FBoenReference, so it knows what skeleton to pick bones from.
-USkeleton* UMLDeformerAsset::GetSkeleton(bool& bInvalidSkeletonIsError)
+USkeleton* UMLDeformerAsset::GetSkeleton(bool& bInvalidSkeletonIsError, const IPropertyHandle* PropertyHandle)
 {
 	bInvalidSkeletonIsError = false;
 #if WITH_EDITOR
diff --git a/Engine/Plugins/Experimental/MLDeformer/Source/MLDeformer/Public/MLDeformerAsset.h b/Engine/Plugins/Experimental/MLDeformer/Source/MLDeformer/Public/MLDeformerAsset.h
--- a/Engine/Plugins/Experimental/MLDeformer/Source/MLDeformer/Public/MLDeformerAsset.h
+++ b/Engine/Plugins/Experimental/MLDeformer/Source/MLDeformer/Public/MLDeformerAsset.h
@@ -20,6 +20,7 @@ class UGeometryCache;
 class UAnimSequence;
 class UMLDeformerAsset;
 class USkeleton;
+class IPropertyHandle;
 
 /** The activation function to use during the ML Deformer training process. */
 UENUM()
@@ -128,7 +129,7 @@ public:
 
 
 	//~IBoneReferenceSkeletonProvider interface
-	USkeleton* GetSkeleton(bool& bInvalidSkeletonIsError) override;
+	USkeleton* GetSkeleton(bool& bInvalidSkeletonIsError, const IPropertyHandle* PropertyHandle) override;
 	//~End of IBoneReferenceSkeletonProvider interface
 
 
