[MLDeformer] Another attempt to fix some error reported by Horde.

diff --git a/Engine/Plugins/Experimental/MLDeformer/Source/MLDeformer/Public/MLDeformerAsset.h b/Engine/Plugins/Experimental/MLDeformer/Source/MLDeformer/Public/MLDeformerAsset.h
--- a/Engine/Plugins/Experimental/MLDeformer/Source/MLDeformer/Public/MLDeformerAsset.h
+++ b/Engine/Plugins/Experimental/MLDeformer/Source/MLDeformer/Public/MLDeformerAsset.h
@@ -110,7 +110,9 @@ struct MLDEFORMER_API FMLDeformerMeshMapping
 UCLASS(BlueprintType, hidecategories=Object)
 class MLDEFORMER_API UMLDeformerAsset 
 	: public UObject
+#ifdef WITH_EDITOR
 	, public IBoneReferenceSkeletonProvider
+#endif
 {
 	GENERATED_BODY()
 
@@ -127,11 +129,11 @@ public:
 #endif
 	//~End of UObject interface
 
-
+#if WITH_EDITOR
 	//~IBoneReferenceSkeletonProvider interface
 	USkeleton* GetSkeleton(bool& bInvalidSkeletonIsError, const IPropertyHandle* PropertyHandle) override;
 	//~End of IBoneReferenceSkeletonProvider interface
-
+#endif
 
 #if WITH_EDITOR
 	FText GetGeomCacheErrorText(UGeometryCache* InGeomCache) const;
