[MLDeformer]
- Added a curve include list.
- Added widgets required for the curve include list.
- Improved compatibility error reporting on screen.
- Got rid of an ensure being triggered because of network input incompatibilities, and show an error on screen on the ML Deformer asset editor viewport instead.

diff --git a/Engine/Plugins/Experimental/MLDeformer/Source/MLDeformer/Private/MLDeformerAsset.cpp b/Engine/Plugins/Experimental/MLDeformer/Source/MLDeformer/Private/MLDeformerAsset.cpp
--- a/Engine/Plugins/Experimental/MLDeformer/Source/MLDeformer/Private/MLDeformerAsset.cpp
+++ b/Engine/Plugins/Experimental/MLDeformer/Source/MLDeformer/Private/MLDeformerAsset.cpp
@@ -13,6 +13,7 @@
 #include "Engine/SkeletalMesh.h"
 #include "Animation/AnimSequence.h"
 #include "Animation/AnimData/AnimDataModel.h"
+#include "CurveReference.h"
 
 #define LOCTEXT_NAMESPACE "MLDeformerAsset"
 
@@ -405,10 +406,93 @@ void UMLDeformerAsset::InitBoneIncludeListToAnimatedBonesOnly()
 	}
 	else
 	{
+		BoneIncludeList.Empty();
 		UE_LOG(LogMLDeformer, Warning, TEXT("There are no animated bone rotations in Anim Sequence '%s'."), *AnimSequence->GetName());
 	}
 }
 
+void UMLDeformerAsset::InitCurveIncludeListToAnimatedCurvesOnly()
+{
+	if (!AnimSequence)
+	{
+		UE_LOG(LogMLDeformer, Warning, TEXT("Cannot initialize curve list as no Anim Sequence has been picked."));
+		return;
+	}
+
+	const UAnimDataModel* DataModel = AnimSequence->GetDataModel();
+	if (!DataModel)
+	{
+		UE_LOG(LogMLDeformer, Warning, TEXT("Anim sequence has no data model."));
+		return;
+	}
+
+	if (!SkeletalMesh)
+	{
+		UE_LOG(LogMLDeformer, Warning, TEXT("Skeletal Mesh has not been set."));
+		return;
+	}
+
+	USkeleton* Skeleton = SkeletalMesh->GetSkeleton();
+	if (!Skeleton)
+	{
+		UE_LOG(LogMLDeformer, Warning, TEXT("Skeletal Mesh has no skeleton."));
+		return;
+	}
+
+	// Iterate over all curves that are both in the skeleton and the animation.
+	TArray<FName> AnimatedCurveList;
+	const FSmartNameMapping* Mapping = Skeleton->GetSmartNameContainer(USkeleton::AnimCurveMappingName);
+	if (Mapping)
+	{
+		TArray<FName> SkeletonCurveNames;
+		Mapping->FillNameArray(SkeletonCurveNames);
+		for (const FName& SkeletonCurveName : SkeletonCurveNames)
+		{
+			const TArray<FFloatCurve>& AnimCurves = DataModel->GetFloatCurves();
+			for (const FFloatCurve& AnimCurve : AnimCurves)
+			{
+				if (AnimCurve.Name.IsValid() && AnimCurve.Name.DisplayName == SkeletonCurveName)
+				{
+					TArray<float> TimeValues;
+					TArray<float> KeyValues;
+					AnimCurve.GetKeys(TimeValues, KeyValues);
+					if (KeyValues.Num() > 0)
+					{
+						const float FirstKeyValue = KeyValues[0];					
+						for (float CurKeyValue : KeyValues)
+						{
+							if (CurKeyValue != FirstKeyValue)
+							{
+								AnimatedCurveList.Add(SkeletonCurveName);
+								break;
+							}
+						}
+					}
+					break;
+				}
+			}
+		}
+	}
+
+	// Init the bone include list using the animated bones.
+	if (!AnimatedCurveList.IsEmpty())
+	{
+		CurveIncludeList.Empty();
+		CurveIncludeList.Reserve(AnimatedCurveList.Num());
+		for (FName CurveName : AnimatedCurveList)
+		{
+			CurveIncludeList.AddDefaulted();
+			FCurveReference& CurveRef = CurveIncludeList.Last();
+			CurveRef.CurveName = CurveName;
+		}
+	}
+	else
+	{
+		CurveIncludeList.Empty();
+		UE_LOG(LogMLDeformer, Warning, TEXT("There are no animated curves in Anim Sequence '%s'."), *AnimSequence->GetName());
+	}
+}
+
 int32 UMLDeformerAsset::GetNumFramesForTraining() const 
 { 
 	return FMath::Min(GetNumFrames(), GetTrainingFrameLimit()); 
@@ -594,16 +678,18 @@ FMLDeformerInputInfo UMLDeformerAsset::CreateInputInfo() const
 	Settings.bIncludeBones = (TrainingInputs == ETrainingInputs::BonesAndCurves || TrainingInputs == ETrainingInputs::BonesOnly);
 	Settings.bIncludeCurves = (TrainingInputs == ETrainingInputs::BonesAndCurves || TrainingInputs == ETrainingInputs::CurvesOnly);
 
+	const USkeleton* Skeleton = SkeletalMesh.Get() ? SkeletalMesh->GetSkeleton() : nullptr;
+
 	// Set the list of bones to use, from the bone references.
-	if (!BoneIncludeList.IsEmpty())
+	if (!BoneIncludeList.IsEmpty() && Skeleton)
 	{
 		FString BoneName;
 		for (const FBoneReference& BoneReference : BoneIncludeList)
 		{
-			if (SkeletalMesh && SkeletalMesh->GetSkeleton() && BoneReference.BoneName.IsValid())
+			if (BoneReference.BoneName.IsValid())
 			{
 				BoneName = BoneReference.BoneName.ToString();
-				const int32 BoneId = SkeletalMesh->GetRefSkeleton().FindRawBoneIndex(BoneReference.BoneName);
+				const int32 BoneId = Skeleton->GetReferenceSkeleton().FindRawBoneIndex(BoneReference.BoneName);
 				if (BoneId != INDEX_NONE && !Settings.BoneNamesToInclude.Contains(BoneName))
 				{
 					Settings.BoneNamesToInclude.Add(BoneName);
@@ -617,6 +703,37 @@ FMLDeformerInputInfo UMLDeformerAsset::CreateInputInfo() const
 		UE_LOG(LogMLDeformer, Display, TEXT("Including ALL bones of skeleton in training."));
 	}
 
+	// Set the list of curves to use, from the curve references.
+	if (!CurveIncludeList.IsEmpty() && Skeleton)
+	{
+		const FSmartNameMapping* Mapping = Skeleton->GetSmartNameContainer(USkeleton::AnimCurveMappingName);
+		if (Mapping)
+		{
+			// Get the list of curve names in the skeleton.
+			TArray<FName> SkeletonCurveNames;
+			Mapping->FillNameArray(SkeletonCurveNames);
+
+			// Add only curves that are also in the skeleton.
+			FString CurveName;
+			for (const FCurveReference& CurveReference : CurveIncludeList)
+			{
+				if (CurveReference.CurveName.IsValid())
+				{
+					CurveName = CurveReference.CurveName.ToString();
+					if (SkeletonCurveNames.Contains(CurveReference.CurveName))
+					{
+						UE_LOG(LogMLDeformer, Display, TEXT("Including curve '%s' in training."), *CurveName);
+						Settings.CurveNamesToInclude.Add(CurveName);
+					}
+				}
+			}
+		}
+	}
+	else
+	{
+		UE_LOG(LogMLDeformer, Display, TEXT("Including ALL bones of skeleton in training."));
+	}
+
 	Result.Init(Settings);
 	return Result;
 }
diff --git a/Engine/Plugins/Experimental/MLDeformer/Source/MLDeformer/Private/MLDeformerInstance.cpp b/Engine/Plugins/Experimental/MLDeformer/Source/MLDeformer/Private/MLDeformerInstance.cpp
--- a/Engine/Plugins/Experimental/MLDeformer/Source/MLDeformer/Private/MLDeformerInstance.cpp
+++ b/Engine/Plugins/Experimental/MLDeformer/Source/MLDeformer/Private/MLDeformerInstance.cpp
@@ -39,16 +39,24 @@ void FMLDeformerInstance::Init(UMLDeformerAsset* Asset, USkeletalMeshComponent*
 	}
 
 	// Perform a compatibility check.
-	bIsCompatible = SkelMesh && CheckCompatibility(SkelMeshComponent, true).IsEmpty();
+	UpdateCompatibilityStatus();
 }
 
-FString FMLDeformerInstance::CheckCompatibility(USkeletalMeshComponent* InSkelMeshComponent, bool LogIssues) const
+void FMLDeformerInstance::UpdateCompatibilityStatus()
 {
+	bIsCompatible = SkeletalMeshComponent->SkeletalMesh && CheckCompatibility(SkeletalMeshComponent, true).IsEmpty();
+}
+
+FString FMLDeformerInstance::CheckCompatibility(USkeletalMeshComponent* InSkelMeshComponent, bool LogIssues)
+{
+	ErrorText = FString();
+
 	// If we're not compatible, generate a compatibility string.
 	USkeletalMesh* SkelMesh = InSkelMeshComponent ? InSkelMeshComponent->SkeletalMesh.Get() : nullptr;
 	if (SkelMesh && !DeformerAsset->GetInputInfo().IsCompatible(SkelMesh))
 	{
-		const FString ErrorText = DeformerAsset->GetInputInfo().GenerateCompatibilityErrorString(SkelMesh);
+		ErrorText += DeformerAsset->GetInputInfo().GenerateCompatibilityErrorString(SkelMesh);
+		ErrorText += "\n";
 		check(!ErrorText.IsEmpty());
 		if (LogIssues)
 		{
@@ -57,11 +65,34 @@ FString FMLDeformerInstance::CheckCompatibility(USkeletalMeshComponent* InSkelMe
 				*SkelMesh->GetName(), 
 				*ErrorText);
 		}
-		return ErrorText;
 	}
 
-	// No errors.
-	return FString();
+	UNeuralNetwork* NeuralNetwork = DeformerAsset->GetInferenceNeuralNetwork();
+	if (NeuralNetwork == nullptr || !NeuralNetwork->IsLoaded())
+	{
+		const FString NetworkErrorString = "The neural network needs to be trained.";
+		ErrorText += NetworkErrorString + "\n";
+		if (LogIssues)
+		{
+			UE_LOG(LogMLDeformer, Error, TEXT("Deformer '%s': %s"), *DeformerAsset->GetName(), *NetworkErrorString);
+		}
+	}
+	else
+	{
+		const int64 NumNeuralNetInputs = NeuralNetwork->GetInputTensor().Num();
+		const int64 NumDeformerAssetInputs = static_cast<int64>(DeformerAsset->GetInputInfo().CalcNumNeuralNetInputs());
+		if (NumNeuralNetInputs != NumDeformerAssetInputs)
+		{
+			const FString InputErrorString = "The number of network inputs doesn't match the asset. Please retrain the asset."; 
+			ErrorText += InputErrorString + "\n";
+			if (LogIssues)
+			{
+				UE_LOG(LogMLDeformer, Error, TEXT("Deformer '%s': %s"), *DeformerAsset->GetName(), *InputErrorString);
+			}
+		}
+	}
+
+	return ErrorText;
 }
 
 int64 FMLDeformerInstance::SetBoneTransforms(float* OutputBuffer, int64 OutputBufferSize, int64 StartIndex)
@@ -167,12 +198,6 @@ void FMLDeformerInstance::Update()
 	// If the neural network expects a different number of inputs, do nothing.
 	const int64 NumNeuralNetInputs = NeuralNetwork->GetInputTensor().Num();
 	const int64 NumDeformerAssetInputs = static_cast<int64>(DeformerAsset->GetInputInfo().CalcNumNeuralNetInputs());
-	ensureMsgf(
-		NumNeuralNetInputs == NumDeformerAssetInputs, 
-		TEXT("Neural network %s expects %d inputs, while the deformer asset expects to feed %d inputs."), 
-		*DeformerAsset->GetName(), 
-		NumNeuralNetInputs, 
-		NumDeformerAssetInputs);
 	if (NumNeuralNetInputs != NumDeformerAssetInputs)
 	{
 		return;
diff --git a/Engine/Plugins/Experimental/MLDeformer/Source/MLDeformer/Public/CurveReference.h b/Engine/Plugins/Experimental/MLDeformer/Source/MLDeformer/Public/CurveReference.h
new file mode 100644
--- /dev/null
+++ b/Engine/Plugins/Experimental/MLDeformer/Source/MLDeformer/Public/CurveReference.h
@@ -0,0 +1,44 @@
+// Copyright Epic Games, Inc. All Rights Reserved.
+
+#pragma once
+
+#include "CoreMinimal.h"
+#include "UObject/ObjectMacros.h"
+#include "CurveReference.generated.h"
+
+USTRUCT()
+struct FCurveReference
+{
+	GENERATED_USTRUCT_BODY()
+
+	FCurveReference(const FName& InCurveName=NAME_None)
+		: CurveName(InCurveName)
+	{
+	}
+
+	bool operator==(const FCurveReference& Other) const
+	{
+		return (CurveName == Other.CurveName);
+	}
+
+	bool operator!=(const FCurveReference& Other) const
+	{
+		return (CurveName != Other.CurveName);
+	}
+
+	friend FArchive& operator<<(FArchive& Ar, FCurveReference& C)
+	{
+		Ar << C.CurveName;
+		return Ar;
+	}
+
+	bool Serialize(FArchive& Ar)
+	{
+		Ar << *this;
+		return true;
+	}
+
+	/** Name of curve. **/
+	UPROPERTY(EditAnywhere, Category = AnimCurveReference)
+	FName CurveName;
+};
diff --git a/Engine/Plugins/Experimental/MLDeformer/Source/MLDeformer/Public/MLDeformerAsset.h b/Engine/Plugins/Experimental/MLDeformer/Source/MLDeformer/Public/MLDeformerAsset.h
--- a/Engine/Plugins/Experimental/MLDeformer/Source/MLDeformer/Public/MLDeformerAsset.h
+++ b/Engine/Plugins/Experimental/MLDeformer/Source/MLDeformer/Public/MLDeformerAsset.h
@@ -7,6 +7,7 @@
 #include "MLDeformerVizSettings.h"
 #include "MLDeformerInputInfo.h"
 #include "BoneContainer.h"
+#include "CurveReference.h"
 #include "RenderCommandFence.h"
 #include "UObject/Object.h"
 #include "RenderResource.h"
@@ -155,6 +156,7 @@ public:
 	FMLDeformerInputInfo CreateInputInfo() const;
 
 	void InitBoneIncludeListToAnimatedBonesOnly();
+	void InitCurveIncludeListToAnimatedCurvesOnly();
 #endif
 
 	void InitVertexMap();
@@ -190,6 +192,9 @@ public:
 	TArray<FBoneReference>& GetBoneIncludeList() { return BoneIncludeList; }
 	const TArray<FBoneReference>& GetBoneIncludeList() const { return BoneIncludeList; }
 
+	TArray<FCurveReference>& GetCurveIncludeList() { return CurveIncludeList; }
+	const TArray<FCurveReference>& GetCurveIncludeList() const { return CurveIncludeList; }
+
 	int32 GetNumFrames() const;
 	int32 GetTrainingFrameLimit() const { return MaxTrainingFrames; }
 	int32 GetNumFramesForTraining() const;
@@ -336,5 +341,9 @@ public:
 	/** The bones to include during training. When none are provided, all bones of the Skeleton will be included. */
 	UPROPERTY(EditAnywhere, Category = "Inputs and Output", meta = (EditCondition = "TrainingInputs == ETrainingInputs::BonesAndCurves || TrainingInputs == ETrainingInputs::BonesOnly"))
 	TArray<FBoneReference> BoneIncludeList;
+
+	/** The curves to include during training. When none are provided, all curves of the Skeleton will be included. */
+	UPROPERTY(EditAnywhere, Category = "Inputs and Output", meta = (EditCondition = "TrainingInputs == ETrainingInputs::BonesAndCurves || TrainingInputs == ETrainingInputs::CurvesOnly"))
+	TArray<FCurveReference> CurveIncludeList;
 #endif // WITH_EDITORONLY_DATA
 };
diff --git a/Engine/Plugins/Experimental/MLDeformer/Source/MLDeformer/Public/MLDeformerComponent.h b/Engine/Plugins/Experimental/MLDeformer/Source/MLDeformer/Public/MLDeformerComponent.h
--- a/Engine/Plugins/Experimental/MLDeformer/Source/MLDeformer/Public/MLDeformerComponent.h
+++ b/Engine/Plugins/Experimental/MLDeformer/Source/MLDeformer/Public/MLDeformerComponent.h
@@ -43,6 +43,7 @@ public:
 	void SetVertexDeltaMultiplier(float ScaleFactor) { VertexDeltaMultiplier = ScaleFactor;  }
 
 	const FMLDeformerInstance& GetDeformerInstance() const { return DeformerInstance; }
+	FMLDeformerInstance& GetDeformerInstance() { return DeformerInstance; }
 
 private:
 	/** The deformer asset to use. */
diff --git a/Engine/Plugins/Experimental/MLDeformer/Source/MLDeformer/Public/MLDeformerInstance.h b/Engine/Plugins/Experimental/MLDeformer/Source/MLDeformer/Public/MLDeformerInstance.h
--- a/Engine/Plugins/Experimental/MLDeformer/Source/MLDeformer/Public/MLDeformerInstance.h
+++ b/Engine/Plugins/Experimental/MLDeformer/Source/MLDeformer/Public/MLDeformerInstance.h
@@ -43,21 +43,31 @@ public:
 	 */
 	bool IsCompatible() const { return bIsCompatible; }
 
+	/**
+	 * Get the compatibility error text. This will be a non-empty string in case IsCompatible() return false. 
+	 * @return Returns the string containing more information about the reasons for compatibility issues.
+	 */
+	const FString& GetCompatibilityErrorText() const { return ErrorText; }
+
 	/**
 	 * Get the skeletal mesh component we're working with.
 	 * @return A pointer to the skeletal mesh component.
 	 */
 	USkeletalMeshComponent* GetSkeletalMeshComponent() const { return SkeletalMeshComponent; }
 
+	/** Update the compatibility status, as returned by IsCompatible() and GetCompatibilityErrorText(). */
+	void UpdateCompatibilityStatus();
+
 private:
 	/**
 	 * Check whether the deformer is compatible with a given skeletal mesh component.
+	 * This internally also edits the value returned by GetCompatibilityErrorText().
 	 * @param InSkelMeshComponent The skeletal mesh component to check compatibility with.
 	 * @param LogIssues Set to true to automatically log any compatibility errors.
 	 * @return Returns the error string. When the returned string is empty, there were no errors and thus
 	 *         the specified skeletal mesh component is compatible.
 	 */
-	FString CheckCompatibility(USkeletalMeshComponent* InSkelMeshComponent, bool LogIssues=false) const;
+	FString CheckCompatibility(USkeletalMeshComponent* InSkelMeshComponent, bool LogIssues=false);
 
 	/**
 	 * Update the neural network input values directly inside its input tensor.
@@ -98,6 +108,9 @@ private:
 	/** Maps the ML deformer asset bone index to a skeletal mesh component bone index. */
 	TArray<int32> AssetBonesToSkelMeshMappings;
 
+	/** The compatibility error text, in case bIsCompatible is false. */
+	FString ErrorText;
+
 	/** Are the deformer asset and the used skeletal mesh component compatible? */
 	bool bIsCompatible = false;
 };
diff --git a/Engine/Plugins/Experimental/MLDeformer/Source/MLDeformerEditor/Private/CurveReferenceCustomization.cpp b/Engine/Plugins/Experimental/MLDeformer/Source/MLDeformerEditor/Private/CurveReferenceCustomization.cpp
new file mode 100644
--- /dev/null
+++ b/Engine/Plugins/Experimental/MLDeformer/Source/MLDeformerEditor/Private/CurveReferenceCustomization.cpp
@@ -0,0 +1,89 @@
+// Copyright Epic Games, Inc. All Rights Reserved.
+
+#include "CurveReferenceCustomization.h"
+#include "CurveReference.h"
+#include "SCurvePickerWidget.h"
+#include "MLDeformerAsset.h"
+
+TSharedRef<IPropertyTypeCustomization> FCurveReferenceCustomization::MakeInstance()
+{
+	return MakeShareable(new FCurveReferenceCustomization());
+}
+
+void FCurveReferenceCustomization::CustomizeHeader(TSharedRef<IPropertyHandle> StructPropertyHandle, class FDetailWidgetRow& HeaderRow, IPropertyTypeCustomizationUtils& StructCustomizationUtils)
+{
+	SetPropertyHandle(StructPropertyHandle);
+	SetSkeleton(StructPropertyHandle);
+
+	if (CurveNameProperty->IsValidHandle())
+	{
+		HeaderRow
+		.NameContent()
+		[
+			StructPropertyHandle->CreatePropertyNameWidget()
+		]
+		.ValueContent()
+		.MaxDesiredWidth(0.0f)
+		[
+			SNew(SCurveSelectionWidget)
+			.OnCurveSelectionChanged(this, &FCurveReferenceCustomization::OnCurveSelectionChanged)
+			.OnGetSelectedCurve(this, &FCurveReferenceCustomization::OnGetSelectedCurve)
+			.OnGetSkeleton(this, &FCurveReferenceCustomization::OnGetSkeleton)
+		];
+	}
+}
+
+void FCurveReferenceCustomization::SetSkeleton(TSharedRef<IPropertyHandle> StructPropertyHandle) 
+{
+	Skeleton = nullptr;
+
+	TArray<UObject*> Objects;
+	StructPropertyHandle->GetOuterObjects(Objects);
+	check(Objects.Num() == 1);
+
+	UMLDeformerAsset* DeformerAsset = Cast<UMLDeformerAsset>(Objects[0]);
+	check(DeformerAsset);
+
+	// Get the skeleton.
+	bool bInvalidSkeletonIsError = false;
+	Skeleton = DeformerAsset->GetSkeleton(bInvalidSkeletonIsError, nullptr);
+}
+
+TSharedPtr<IPropertyHandle> FCurveReferenceCustomization::FindStructMemberProperty(TSharedRef<IPropertyHandle> PropertyHandle, const FName& PropertyName)
+{
+	uint32 NumChildren = 0;
+	PropertyHandle->GetNumChildren(NumChildren);
+	for (uint32 ChildIdx = 0; ChildIdx < NumChildren; ++ChildIdx)
+	{
+		TSharedPtr<IPropertyHandle> ChildHandle = PropertyHandle->GetChildHandle(ChildIdx);
+		if (ChildHandle->GetProperty()->GetFName() == PropertyName)
+		{
+			return ChildHandle;
+		}
+	}
+
+	return TSharedPtr<IPropertyHandle>();
+}
+
+void FCurveReferenceCustomization::SetPropertyHandle(TSharedRef<IPropertyHandle> StructPropertyHandle)
+{
+	CurveNameProperty = FindStructMemberProperty(StructPropertyHandle, GET_MEMBER_NAME_CHECKED(FCurveReference, CurveName));
+	check(CurveNameProperty->IsValidHandle());
+}
+
+void FCurveReferenceCustomization::OnCurveSelectionChanged(const FString& Name)
+{
+	CurveNameProperty->SetValue(Name);
+}
+
+FString FCurveReferenceCustomization::OnGetSelectedCurve() const
+{
+	FString CurveName;
+	CurveNameProperty->GetValue(CurveName);
+	return CurveName;
+}
+
+USkeleton* FCurveReferenceCustomization::OnGetSkeleton() const
+{
+	return Skeleton;
+}
diff --git a/Engine/Plugins/Experimental/MLDeformer/Source/MLDeformerEditor/Private/MLDeformerAssetDetails.cpp b/Engine/Plugins/Experimental/MLDeformer/Source/MLDeformerEditor/Private/MLDeformerAssetDetails.cpp
--- a/Engine/Plugins/Experimental/MLDeformer/Source/MLDeformerEditor/Private/MLDeformerAssetDetails.cpp
+++ b/Engine/Plugins/Experimental/MLDeformer/Source/MLDeformerEditor/Private/MLDeformerAssetDetails.cpp
@@ -239,11 +239,9 @@ void FMLDeformerAssetDetails::CustomizeDetails(class IDetailLayoutBuilder& Detai
 	InputOutputCategoryBuilder.AddProperty(GET_MEMBER_NAME_CHECKED(UMLDeformerAsset, DeltaCutoffLength));
 
 	// Bone include list group.
-	IDetailGroup& BoneIncludeGroup = InputOutputCategoryBuilder.AddGroup("BoneIncludeGroup", LOCTEXT("BoneIncludeGroup", "Bones"), false, true);
+	IDetailGroup& BoneIncludeGroup = InputOutputCategoryBuilder.AddGroup("BoneIncludeGroup", LOCTEXT("BoneIncludeGroup", "Bones"), false, false);
 	BoneIncludeGroup.AddWidgetRow()
-		//.WholeRowContent()
 		.ValueContent()
-		//.NameContent()
 		[
 			SNew(SButton)
 			.HAlign(HAlign_Center)
@@ -254,6 +252,21 @@ void FMLDeformerAssetDetails::CustomizeDetails(class IDetailLayoutBuilder& Detai
 		];
 	BoneIncludeGroup.AddPropertyRow(DetailBuilder.GetProperty(GET_MEMBER_NAME_CHECKED(UMLDeformerAsset, BoneIncludeList)));
 
+	// Curve include list group.
+	IDetailGroup& CurveIncludeGroup = InputOutputCategoryBuilder.AddGroup("CurveIncludeGroup", LOCTEXT("CurveIncludeGroup", "Curves"), false, false);
+	CurveIncludeGroup.AddWidgetRow()
+		.ValueContent()
+		[
+			SNew(SButton)
+			.HAlign(HAlign_Center)
+			.VAlign(VAlign_Center)
+			.Text(LOCTEXT("AnimatedCurvesButton", "Animated Curves Only"))
+			.OnClicked(FOnClicked::CreateSP(this, &FMLDeformerAssetDetails::OnFilterAnimatedCurvesOnly, DeformerAsset))
+			.IsEnabled_Lambda([DeformerAsset](){ return (DeformerAsset->GetTrainingInputs() == ETrainingInputs::BonesAndCurves) || (DeformerAsset->GetTrainingInputs() == ETrainingInputs::CurvesOnly); })
+		];
+	CurveIncludeGroup.AddPropertyRow(DetailBuilder.GetProperty(GET_MEMBER_NAME_CHECKED(UMLDeformerAsset, CurveIncludeList)));
+
+
 	// Training settings.
 	IDetailCategoryBuilder& SettingsCategoryBuilder = DetailBuilder.EditCategory("Training Settings", FText::GetEmpty(), ECategoryPriority::Important);
 	SettingsCategoryBuilder.AddProperty(GET_MEMBER_NAME_CHECKED(UMLDeformerAsset, NumHiddenLayers));
@@ -327,4 +340,11 @@ FReply FMLDeformerAssetDetails::OnFilterAnimatedBonesOnly(UMLDeformerAsset* Defo
 	return FReply::Handled();
 }
 
+FReply FMLDeformerAssetDetails::OnFilterAnimatedCurvesOnly(UMLDeformerAsset* DeformerAsset) const
+{
+	DeformerAsset->InitCurveIncludeListToAnimatedCurvesOnly();
+	DetailLayoutBuilder->ForceRefreshDetails();
+	return FReply::Handled();
+}
+
 #undef LOCTEXT_NAMESPACE
diff --git a/Engine/Plugins/Experimental/MLDeformer/Source/MLDeformerEditor/Private/MLDeformerAssetDetails.h b/Engine/Plugins/Experimental/MLDeformer/Source/MLDeformerEditor/Private/MLDeformerAssetDetails.h
--- a/Engine/Plugins/Experimental/MLDeformer/Source/MLDeformerEditor/Private/MLDeformerAssetDetails.h
+++ b/Engine/Plugins/Experimental/MLDeformer/Source/MLDeformerEditor/Private/MLDeformerAssetDetails.h
@@ -25,6 +25,7 @@ public:
 private:
 	bool FilterAnimSequences(const FAssetData& AssetData, USkeleton* Skeleton);
 	FReply OnFilterAnimatedBonesOnly(UMLDeformerAsset* DeformerAsset) const;
+	FReply OnFilterAnimatedCurvesOnly(UMLDeformerAsset* DeformerAsset) const;
 
 private:
 	/** Associated detail layout builder. */
diff --git a/Engine/Plugins/Experimental/MLDeformer/Source/MLDeformerEditor/Private/MLDeformerEditor.cpp b/Engine/Plugins/Experimental/MLDeformer/Source/MLDeformerEditor/Private/MLDeformerEditor.cpp
--- a/Engine/Plugins/Experimental/MLDeformer/Source/MLDeformerEditor/Private/MLDeformerEditor.cpp
+++ b/Engine/Plugins/Experimental/MLDeformer/Source/MLDeformerEditor/Private/MLDeformerEditor.cpp
@@ -8,6 +8,8 @@
 #include "MLDeformerEditorMode.h"
 #include "MLDeformer.h"
 
+#include "CurveReferenceCustomization.h"
+
 #include "Modules/ModuleManager.h"
 #include "EditorModeRegistry.h"
 #include "PropertyEditorDelegates.h"
@@ -50,6 +52,7 @@ void FMLDeformerEditor::StartupModule()
 	FPropertyEditorModule& PropertyModule = FModuleManager::GetModuleChecked<FPropertyEditorModule>("PropertyEditor");
 	PropertyModule.RegisterCustomClassLayout("MLDeformerAsset", FOnGetDetailCustomizationInstance::CreateStatic(&FMLDeformerAssetDetails::MakeInstance));
 	PropertyModule.RegisterCustomClassLayout("MLDeformerVizSettings", FOnGetDetailCustomizationInstance::CreateStatic(&FMLDeformerVizSettingsDetails::MakeInstance));
+	PropertyModule.RegisterCustomPropertyTypeLayout("CurveReference", FOnGetPropertyTypeCustomizationInstance::CreateStatic(&FCurveReferenceCustomization::MakeInstance) );
 	PropertyModule.NotifyCustomizationModuleChanged();
 }
 
diff --git a/Engine/Plugins/Experimental/MLDeformer/Source/MLDeformerEditor/Private/MLDeformerEditorData.cpp b/Engine/Plugins/Experimental/MLDeformer/Source/MLDeformerEditor/Private/MLDeformerEditorData.cpp
--- a/Engine/Plugins/Experimental/MLDeformer/Source/MLDeformerEditor/Private/MLDeformerEditorData.cpp
+++ b/Engine/Plugins/Experimental/MLDeformer/Source/MLDeformerEditor/Private/MLDeformerEditorData.cpp
@@ -291,6 +291,12 @@ void FMLDeformerEditorData::InitAssets()
 
 	ClampFrameIndex();
 
+	UMLDeformerComponent* DeformerComponent = GetEditorActor(EMLDeformerEditorActorIndex::DeformedTest).MLDeformerComponent;
+	if (DeformerComponent)
+	{
+		DeformerComponent->GetDeformerInstance().UpdateCompatibilityStatus();
+	}
+
 	// Reinit the single frame cache.
 	FMLDeformerFrameCache::FInitSettings InitSettings;
 	InitSettings.CacheSizeInBytes = 0;	// This makes it just just use one frame, which is the minimum.
@@ -514,7 +520,7 @@ FText FMLDeformerEditorData::GetOverlayText()
 			DeformerInstance.GetSkeletalMeshComponent()->SkeletalMesh &&
 			!DeformerComponent->GetDeformerInstance().IsCompatible() )
 		{
-			return FText::FromString(TEXT("Trained deformer asset is not compatible with skeletal mesh!"));
+			return FText::FromString( DeformerComponent->GetDeformerInstance().GetCompatibilityErrorText() );
 		}
 	}
 
diff --git a/Engine/Plugins/Experimental/MLDeformer/Source/MLDeformerEditor/Private/SCurvePickerWidget.cpp b/Engine/Plugins/Experimental/MLDeformer/Source/MLDeformerEditor/Private/SCurvePickerWidget.cpp
new file mode 100644
--- /dev/null
+++ b/Engine/Plugins/Experimental/MLDeformer/Source/MLDeformerEditor/Private/SCurvePickerWidget.cpp
@@ -0,0 +1,193 @@
+// Copyright Epic Games, Inc. All Rights Reserved.
+
+#include "SCurvePickerWidget.h"
+#include "EditorStyleSet.h"
+#include "DetailLayoutBuilder.h"
+#include "Widgets/Layout/SSeparator.h"
+#include "Widgets/Input/SComboButton.h"
+#include "Widgets/Input/SSearchBox.h"
+#include "Widgets/Views/SListView.h"
+#include "Widgets/SBoxPanel.h"
+#include "Widgets/Text/STextBlock.h"
+#include "Widgets/Layout/SBox.h"
+#include "Animation/Skeleton.h"
+
+#define LOCTEXT_NAMESPACE "SCurvePickerWidget"
+
+void SCurvePickerWidget::Construct(const FArguments& InArgs)
+{
+	OnCurveNamePicked = InArgs._OnCurveNamePicked;
+	OnGetSkeleton = InArgs._OnGetSkeleton;
+
+	if (OnGetSkeleton.IsBound())
+	{
+		Skeleton = OnGetSkeleton.Execute();
+	}
+
+	ChildSlot
+	[
+		SNew(SVerticalBox)
+		+SVerticalBox::Slot()
+		.AutoHeight()
+		[
+			SAssignNew(SearchBox, SSearchBox)
+				.HintText(LOCTEXT("SearchBoxHint", "Search Curves"))
+				.OnTextChanged(this, &SCurvePickerWidget::HandleFilterTextChanged)
+		]
+		+SVerticalBox::Slot()
+		.FillHeight(1.0f)
+		[
+			SAssignNew(NameListView, SListView<TSharedPtr<FString>>)
+			.SelectionMode(ESelectionMode::Single)
+			.ItemHeight(20.0f)
+			.ListItemsSource(&CurveNames)
+			.OnSelectionChanged(this, &SCurvePickerWidget::HandleSelectionChanged)
+			.OnGenerateRow(this, &SCurvePickerWidget::HandleGenerateRow)
+		]
+	];
+
+	RefreshListItems();
+}
+
+void SCurvePickerWidget::HandleSelectionChanged(TSharedPtr<FString> InItem, ESelectInfo::Type InSelectionType)
+{
+	OnCurveNamePicked.ExecuteIfBound(**InItem.Get());
+}
+
+TSharedRef<ITableRow> SCurvePickerWidget::HandleGenerateRow(TSharedPtr<FString> InItem, const TSharedRef<STableViewBase>& InOwnerTable)
+{
+	return 
+		SNew(STableRow<TSharedPtr<FString>>, InOwnerTable)
+		[
+			SNew(SBox)
+			.MinDesiredHeight(20.0f)
+			.VAlign(VAlign_Center)
+			[
+				SNew(STextBlock)
+				.Text(FText::FromString(*InItem.Get()))
+				.HighlightText_Lambda([this]() { return FText::FromString(FilterText); })
+			]
+		];
+}
+
+void SCurvePickerWidget::RefreshListItems()
+{
+	CurveNames.Empty();
+	UniqueCurveNames.Empty();
+
+	// Get the curve names.
+	if (Skeleton)
+	{
+		const FSmartNameMapping* Mapping = Skeleton->GetSmartNameContainer(USkeleton::AnimCurveMappingName);
+		if (Mapping)
+		{
+			TArray<FName> SkeletonCurveNames;
+			Mapping->FillNameArray(SkeletonCurveNames);
+
+			UniqueCurveNames.Reserve(SkeletonCurveNames.Num());
+			for (const FName& CurveName : SkeletonCurveNames)
+			{
+				UniqueCurveNames.Add(CurveName.ToString());
+			}
+		}
+	}
+
+	FilterAvailableCurves();
+}
+
+void SCurvePickerWidget::FilterAvailableCurves()
+{
+	CurveNames.Empty();
+
+	for (const FString& UniqueCurveName : UniqueCurveNames)
+	{
+		if (FilterText.IsEmpty() || UniqueCurveName.Contains(FilterText))
+		{
+			CurveNames.Add(MakeShared<FString>(UniqueCurveName));
+		}
+	}
+
+	NameListView->RequestListRefresh();
+}
+
+void SCurvePickerWidget::HandleFilterTextChanged(const FText& InFilterText)
+{
+	FilterText = InFilterText.ToString();
+	FilterAvailableCurves();
+}
+
+/////////////////////////////////////////////////////
+
+void SCurveSelectionWidget::Construct(const FArguments& InArgs)
+{
+	OnCurveSelectionChanged = InArgs._OnCurveSelectionChanged;
+	OnGetSelectedCurve = InArgs._OnGetSelectedCurve;
+	OnGetSkeleton = InArgs._OnGetSkeleton;
+
+	SuppliedToolTip = InArgs._ToolTipText.Get();
+
+	this->ChildSlot
+	[
+		SAssignNew(CurvePickerButton, SComboButton)
+		.OnGetMenuContent(FOnGetContent::CreateSP(this, &SCurveSelectionWidget::CreateSkeletonWidgetMenu))
+		.ContentPadding(FMargin(4.0f, 2.0f, 4.0f, 2.0f))
+		.ButtonContent()
+		[
+			SNew(STextBlock)
+				.Text(this, &SCurveSelectionWidget::GetCurrentCurveName)
+				.Font(IDetailLayoutBuilder::GetDetailFont())
+				.ToolTipText(this, &SCurveSelectionWidget::GetFinalToolTip)
+		]
+	];
+}
+
+TSharedRef<SWidget> SCurveSelectionWidget::CreateSkeletonWidgetMenu()
+{
+	bool bMultipleValues = false;
+
+	TSharedRef<SCurvePickerWidget> ListWidget = SNew(SCurvePickerWidget)
+		.OnCurveNamePicked(this, &SCurveSelectionWidget::OnSelectionChanged)
+		.OnGetSkeleton(OnGetSkeleton);
+
+	CurvePickerButton->SetMenuContentWidgetToFocus(ListWidget->GetFilterTextWidget());
+
+	return ListWidget;
+}
+
+void SCurveSelectionWidget::OnSelectionChanged(const FString& CurveName)
+{
+	SelectedCurve = FText::FromString(CurveName);
+	if (OnCurveSelectionChanged.IsBound())
+	{
+		OnCurveSelectionChanged.Execute(CurveName);
+	}
+	CurvePickerButton->SetIsOpen(false);
+}
+
+FText SCurveSelectionWidget::GetCurrentCurveName() const
+{
+	if (OnGetSelectedCurve.IsBound())
+	{
+		const FString Name = OnGetSelectedCurve.Execute();
+		return FText::FromString(Name);
+	}
+
+	return FText::GetEmpty();
+}
+
+FText SCurveSelectionWidget::GetFinalToolTip() const
+{
+	return FText::Format(LOCTEXT("CurveClickToolTip", "Curve: {0}\n\n{1}"), GetCurrentCurveName(), SuppliedToolTip);
+}
+
+USkeleton* SCurveSelectionWidget::GetSkeleton() const
+{
+	if (OnGetSkeleton.IsBound())
+	{
+		return OnGetSkeleton.Execute();
+	}
+
+	return nullptr;
+}
+
+#undef LOCTEXT_NAMESPACE
diff --git a/Engine/Plugins/Experimental/MLDeformer/Source/MLDeformerEditor/Public/CurveReferenceCustomization.h b/Engine/Plugins/Experimental/MLDeformer/Source/MLDeformerEditor/Public/CurveReferenceCustomization.h
new file mode 100644
--- /dev/null
+++ b/Engine/Plugins/Experimental/MLDeformer/Source/MLDeformerEditor/Public/CurveReferenceCustomization.h
@@ -0,0 +1,45 @@
+// Copyright Epic Games, Inc. All Rights Reserved.
+
+#pragma once
+
+#include "CoreMinimal.h"
+#include "Layout/Visibility.h"
+#include "Widgets/DeclarativeSyntaxSupport.h"
+#include "Input/Reply.h"
+#include "Widgets/SWidget.h"
+#include "IPropertyTypeCustomization.h"
+#include "PropertyHandle.h"
+#include "IDetailCustomization.h"
+#include "Widgets/Views/STableViewBase.h"
+#include "Widgets/Views/STableRow.h"
+#include "IDetailPropertyExtensionHandler.h"
+#include "UObject/Object.h"
+
+class USkeleton;
+
+class FCurveReferenceCustomization : public IPropertyTypeCustomization
+{
+public:
+	static TSharedRef<IPropertyTypeCustomization> MakeInstance();
+
+	// IPropertyTypeCustomization interface
+	virtual void CustomizeHeader(TSharedRef<IPropertyHandle> StructPropertyHandle, class FDetailWidgetRow& HeaderRow, IPropertyTypeCustomizationUtils& StructCustomizationUtils) override;
+	virtual void CustomizeChildren(TSharedRef<IPropertyHandle> StructPropertyHandle, class IDetailChildrenBuilder& ChildBuilder, IPropertyTypeCustomizationUtils& StructCustomizationUtils) override {}
+
+protected:
+	void SetSkeleton(TSharedRef<IPropertyHandle> StructPropertyHandle);
+	virtual void SetPropertyHandle(TSharedRef<IPropertyHandle> StructPropertyHandle);
+	TSharedPtr<IPropertyHandle> FindStructMemberProperty(TSharedRef<IPropertyHandle> PropertyHandle, const FName& PropertyName);
+
+	// Property to change after curve has been picked.
+	TSharedPtr<IPropertyHandle> CurveNameProperty;
+
+	// The Skeleton we get the curves from.
+	TObjectPtr<USkeleton> Skeleton;
+
+private:
+	// Curve widget delegates
+	virtual void OnCurveSelectionChanged(const FString& Name);
+	virtual FString OnGetSelectedCurve() const;
+	virtual USkeleton* OnGetSkeleton() const;
+};
diff --git a/Engine/Plugins/Experimental/MLDeformer/Source/MLDeformerEditor/Public/SCurvePickerWidget.h b/Engine/Plugins/Experimental/MLDeformer/Source/MLDeformerEditor/Public/SCurvePickerWidget.h
new file mode 100644
--- /dev/null
+++ b/Engine/Plugins/Experimental/MLDeformer/Source/MLDeformerEditor/Public/SCurvePickerWidget.h
@@ -0,0 +1,97 @@
+// Copyright Epic Games, Inc. All Rights Reserved.
+
+#pragma once
+
+#include "CoreMinimal.h"
+#include "UObject/Object.h"
+#include "SlateFwd.h"
+#include "Widgets/SCompoundWidget.h"
+#include "Widgets/DeclarativeSyntaxSupport.h"
+#include "Widgets/SCompoundWidget.h"
+#include "Widgets/Views/STableRow.h"
+
+class SComboButton;
+class SSearchBox;
+class USkeleton;
+
+DECLARE_DELEGATE_OneParam(FOnCurveSelectionChanged, const FString&);
+
+DECLARE_DELEGATE_RetVal(FString, FOnGetSelectedCurve);
+DECLARE_DELEGATE_RetVal(USkeleton*, FOnGetSkeleton);
+
+DECLARE_DELEGATE_OneParam(FOnCurveNamePicked, const FString&)
+
+class SCurvePickerWidget : public SCompoundWidget
+{
+public:
+	SLATE_BEGIN_ARGS(SCurvePickerWidget) {}
+	SLATE_EVENT(FOnCurveNamePicked, OnCurveNamePicked)
+	SLATE_EVENT(FOnGetSkeleton, OnGetSkeleton)
+	SLATE_END_ARGS()
+
+	void Construct(const FArguments& InArgs);
+
+	TSharedPtr<SSearchBox> GetFilterTextWidget() const { return SearchBox; }
+
+private:
+	void RefreshListItems();
+	void FilterAvailableCurves();
+
+	/** UI handlers */
+	void HandleSelectionChanged(TSharedPtr<FString> InItem, ESelectInfo::Type InSelectionType);
+	TSharedRef<ITableRow> HandleGenerateRow(TSharedPtr<FString> InItem, const TSharedRef<STableViewBase>& InOwnerTable);
+	void HandleFilterTextChanged(const FText& InFilterText);
+
+private:
+	/** Delegate fired when a curve name is picked */
+	FOnCurveNamePicked OnCurveNamePicked;
+
+	/** Provide us with a skeleton. */
+	FOnGetSkeleton OnGetSkeleton;
+
+	/** The search filter box. */
+	TSharedPtr<SSearchBox> SearchBox;
+
+	/** The skeleton to get the curves from. */
+	TObjectPtr<USkeleton> Skeleton;
+
+	/** The names of the curves we are displaying */
+	TArray<TSharedPtr<FString>> CurveNames;
+
+	/** All the unique curve names we can find */
+	TSet<FString> UniqueCurveNames;
+
+	/** The string we use to filter curve names */
+	FString FilterText;
+
+	/** The list view used to display names */
+	TSharedPtr<SListView<TSharedPtr<FString>>> NameListView;
+};
+
+
+class SCurveSelectionWidget : public SCompoundWidget
+{
+public: 
+	SLATE_BEGIN_ARGS(SCurveSelectionWidget) {}
+		SLATE_EVENT(FOnCurveSelectionChanged, OnCurveSelectionChanged);
+		SLATE_EVENT(FOnGetSelectedCurve, OnGetSelectedCurve);
+		SLATE_EVENT(FOnGetSkeleton, OnGetSkeleton)
+	SLATE_END_ARGS();
+
+	void Construct(const FArguments& InArgs);
+
+private: 
+	TSharedRef<SWidget> CreateSkeletonWidgetMenu();
+	void OnSelectionChanged(const FString& CurveName);
+	USkeleton* GetSkeleton() const;
+	FText GetCurrentCurveName() const;
+	FText GetFinalToolTip() const;
+
+private:
+	TSharedPtr<SComboButton> CurvePickerButton;
+	FOnCurveSelectionChanged OnCurveSelectionChanged;
+	FOnGetSelectedCurve OnGetSelectedCurve;
+	FOnGetSkeleton OnGetSkeleton;
+	FText SuppliedToolTip;
+	FText SelectedCurve;
+};
