[MLDeformer] Compile error fix.

diff --git a/Engine/Plugins/Experimental/MLDeformer/Source/MLDeformer/Public/MLDeformerAsset.h b/Engine/Plugins/Experimental/MLDeformer/Source/MLDeformer/Public/MLDeformerAsset.h
--- a/Engine/Plugins/Experimental/MLDeformer/Source/MLDeformer/Public/MLDeformerAsset.h
+++ b/Engine/Plugins/Experimental/MLDeformer/Source/MLDeformer/Public/MLDeformerAsset.h
@@ -111,7 +111,7 @@ struct MLDEFORMER_API FMLDeformerMeshMapping
 UCLASS(BlueprintType, hidecategories=Object)
 class MLDEFORMER_API UMLDeformerAsset 
 	: public UObject
-#ifdef WITH_EDITOR
+#if WITH_EDITOR
 	, public IBoneReferenceSkeletonProvider
 #endif
 {
