[MLDeformer]
- Set the noise amount to 0 on default, to prevent very slow training speeds.
- Show a warning when noise is larger than 0, that training speed will be slow.
- Fix some issue in progress window during training, where loss wasn't shown.
- Fixed some bug in the new cache that caused it to return incorrect cached data and that caused certain frames to be recalculated multiple times. Now all epochs after the first epoch should go fast and at the same speed, rather than suddenly speeding up more and more.
- Fixed GPU from holding memory. Releasing memory back to the GPU now after training.

diff --git a/Engine/Plugins/Experimental/MLDeformer/Source/MLDeformer/Public/MLDeformerAsset.h b/Engine/Plugins/Experimental/MLDeformer/Source/MLDeformer/Public/MLDeformerAsset.h
--- a/Engine/Plugins/Experimental/MLDeformer/Source/MLDeformer/Public/MLDeformerAsset.h
+++ b/Engine/Plugins/Experimental/MLDeformer/Source/MLDeformer/Public/MLDeformerAsset.h
@@ -196,6 +196,7 @@ public:
 	int32 GetNumFramesForTraining() const;
 	int32 GetCacheSizeInMegabytes() const { return CacheSizeInMegabytes; }
 	float GetDeltaCutoffLength() const { return DeltaCutoffLength; }
+	float GetNoiseAmount() const { return NoiseAmount; }
 #endif
 
 public:
@@ -317,9 +318,10 @@ public:
 	UPROPERTY(EditAnywhere, BlueprintReadWrite, AdvancedDisplay, Category = "Training Settings", meta = (EditCondition = "LossFunction == ELossFunction::Shrinkage", ClampMin = "0.0"))
 	float ShrinkageThreshold = 0.1f;
 
-	/** The percentage of noise to add. */
-	UPROPERTY(EditAnywhere, BlueprintReadWrite, AdvancedDisplay, Category = "Training Settings", meta = (ClampMin = "0.0", ClampMax = "100.0")) // A percentage.
-	float NoiseAmount = 0.5;
+	/** The percentage of noise to add. Adding some noise can make the system generalize better. Generally you want a low percentage, such as 0.5.
+	  * Currently using noise will disable the caching system, which will slow down training a lot. */
+	UPROPERTY(EditAnywhere, BlueprintReadWrite, AdvancedDisplay, Category = "Training Settings", meta = (ClampMin = "0.0", ClampMax = "5.0")) // A percentage.
+	float NoiseAmount = 0.0;
 
 	/** The loss function to use during model training. */
 	UPROPERTY(EditAnywhere, BlueprintReadWrite, AdvancedDisplay, Category = "Training Settings")
diff --git a/Engine/Plugins/Experimental/MLDeformer/Source/MLDeformerEditor/Private/MLDeformerApplicationMode.cpp b/Engine/Plugins/Experimental/MLDeformer/Source/MLDeformerEditor/Private/MLDeformerApplicationMode.cpp
--- a/Engine/Plugins/Experimental/MLDeformer/Source/MLDeformerEditor/Private/MLDeformerApplicationMode.cpp
+++ b/Engine/Plugins/Experimental/MLDeformer/Source/MLDeformerEditor/Private/MLDeformerApplicationMode.cpp
@@ -46,7 +46,7 @@ FMLDeformerApplicationMode::FMLDeformerApplicationMode(
 	TabFactories.RegisterFactory(MakeShared<FMLDeformerTimelineTabSummoner>(MLDeformerEditor));
 
 	// Create tab layout.
-	TabLayout = FTabManager::NewLayout("Standalone_MLDeformerEditor_Layout_v0.028")
+	TabLayout = FTabManager::NewLayout("Standalone_MLDeformerEditor_Layout_v0.029")
 		->AddArea
 		(
 			FTabManager::NewPrimaryArea()
diff --git a/Engine/Plugins/Experimental/MLDeformer/Source/MLDeformerEditor/Private/MLDeformerAssetDetails.cpp b/Engine/Plugins/Experimental/MLDeformer/Source/MLDeformerEditor/Private/MLDeformerAssetDetails.cpp
--- a/Engine/Plugins/Experimental/MLDeformer/Source/MLDeformerEditor/Private/MLDeformerAssetDetails.cpp
+++ b/Engine/Plugins/Experimental/MLDeformer/Source/MLDeformerEditor/Private/MLDeformerAssetDetails.cpp
@@ -320,6 +320,24 @@ void FMLDeformerAssetDetails::CustomizeDetails(class IDetailLayoutBuilder& Detai
 	SettingsCategoryBuilder.AddProperty(GET_MEMBER_NAME_CHECKED(UMLDeformerAsset, ShrinkageSpeed));
 	SettingsCategoryBuilder.AddProperty(GET_MEMBER_NAME_CHECKED(UMLDeformerAsset, ShrinkageThreshold));
 	SettingsCategoryBuilder.AddProperty(GET_MEMBER_NAME_CHECKED(UMLDeformerAsset, NoiseAmount));
+	if (DeformerAsset)
+	{
+		// Check if the noise amount is greater than 0, if so, show a warning about slow training times.
+		const FText WarningText = LOCTEXT("NoiseWarning", "Adding noise will disable caching, which will slow down training a lot.");
+		FDetailWidgetRow& WarningRow = SettingsCategoryBuilder.AddCustomRow(FText::FromString("NoiseWarning"), true)
+			.Visibility((DeformerAsset->GetNoiseAmount() > 0.0f) ? EVisibility::Visible : EVisibility::Collapsed)
+			.WholeRowContent()
+			[
+				SNew(SBox)
+				.Padding(FMargin(0.0f, 4.0f))
+				[
+					SNew(SWarningOrErrorBox)
+					.MessageStyle(EMessageStyle::Warning)
+					.Message(WarningText)
+				]
+			];
+	}
+
 	SettingsCategoryBuilder.AddProperty(GET_MEMBER_NAME_CHECKED(UMLDeformerAsset, DeviceType));
 }
 
diff --git a/Engine/Plugins/Experimental/MLDeformer/Source/MLDeformerEditor/Private/MLDeformerEditorToolkit.cpp b/Engine/Plugins/Experimental/MLDeformer/Source/MLDeformerEditor/Private/MLDeformerEditorToolkit.cpp
--- a/Engine/Plugins/Experimental/MLDeformer/Source/MLDeformerEditor/Private/MLDeformerEditorToolkit.cpp
+++ b/Engine/Plugins/Experimental/MLDeformer/Source/MLDeformerEditor/Private/MLDeformerEditorToolkit.cpp
@@ -726,6 +726,11 @@ void FMLDeformerEditorToolkit::OnFinishedChangingDetails(const FPropertyChangedE
 		}
 	}
 	else
+	if (Property->GetFName() == GET_MEMBER_NAME_CHECKED(UMLDeformerAsset, NoiseAmount))
+	{
+		EditorData->GetDetailsView()->ForceRefresh();
+	}
+	else
 	if (Property->GetFName() == GET_MEMBER_NAME_CHECKED(UMLDeformerVizSettings, AnimPlaySpeed))
 	{
 		EditorData->UpdateTestAnimPlaySpeed();
