MLDeformer: Fix reading past end of GroundTruth buffer.
When GroundTruth mode was disabled this was still feeding NaN into the shader.

diff --git a/Engine/Plugins/Experimental/MLDeformer/Shaders/Private/MLDeformerDebugDataInterface.ush b/Engine/Plugins/Experimental/MLDeformer/Shaders/Private/MLDeformerDebugDataInterface.ush
--- a/Engine/Plugins/Experimental/MLDeformer/Shaders/Private/MLDeformerDebugDataInterface.ush
+++ b/Engine/Plugins/Experimental/MLDeformer/Shaders/Private/MLDeformerDebugDataInterface.ush
@@ -7,6 +7,7 @@ uint DI_LOCAL(InputStreamStart);
 uint DI_LOCAL(HeatMapMode);
 float DI_LOCAL(HeatMapScale);
 float DI_LOCAL(GroundTruthLerp);
+uint DI_LOCAL(GroundTruthBufferSize);
 
 StructuredBuffer<float3> DI_LOCAL(PositionGroundTruthBuffer);
 Buffer<uint> DI_LOCAL(VertexMapBuffer);
@@ -19,6 +20,10 @@ DI_IMPL_READ(ReadNumVertices, uint, void)
 DI_IMPL_READ(ReadPositionGroundTruth, float3, uint VertexIndex)
 {
 	uint BufferIndex = DI_LOCAL(VertexMapBuffer)[VertexIndex + DI_LOCAL(InputStreamStart)];
+	if (BufferIndex >= DI_LOCAL(GroundTruthBufferSize))
+	{
+		return float3(0, 0, 0);
+	}
 	return DI_LOCAL(PositionGroundTruthBuffer)[BufferIndex];
 }
 
diff --git a/Engine/Plugins/Experimental/MLDeformer/Source/MLDeformer/Private/MLDeformerDebugDataInterface.cpp b/Engine/Plugins/Experimental/MLDeformer/Source/MLDeformer/Private/MLDeformerDebugDataInterface.cpp
--- a/Engine/Plugins/Experimental/MLDeformer/Source/MLDeformer/Private/MLDeformerDebugDataInterface.cpp
+++ b/Engine/Plugins/Experimental/MLDeformer/Source/MLDeformer/Private/MLDeformerDebugDataInterface.cpp
@@ -92,6 +92,7 @@ BEGIN_SHADER_PARAMETER_STRUCT(FMLDeformerDebugDataInterfaceParameters, )
 	SHADER_PARAMETER(int32, HeatMapMode)
 	SHADER_PARAMETER(float, HeatMapScale)
 	SHADER_PARAMETER(float, GroundTruthLerp)
+	SHADER_PARAMETER(uint32, GroundTruthBufferSize)
 	SHADER_PARAMETER_RDG_BUFFER_SRV(StructuredBuffer<float3>, PositionGroundTruthBuffer)
 	SHADER_PARAMETER_SRV(Buffer<uint>, VertexMapBuffer)
 END_SHADER_PARAMETER_STRUCT()
@@ -290,6 +291,7 @@ void FMLDeformerDebugDataProviderProxy::GetBindings(int32 InvocationIndex, TCHAR
 	Parameters.HeatMapMode = HeatMapMode;
 	Parameters.HeatMapScale = HeatMapScale;
 	Parameters.GroundTruthLerp = GroundTruthLerp;
+	Parameters.GroundTruthBufferSize = GroundTruthPositions.Num();
 	Parameters.PositionGroundTruthBuffer = GroundTruthBufferSRV;
 	Parameters.VertexMapBuffer = VertexMapBufferSRV;
 
