MLDeformer: Make quaternions canonical in training and inference.
Fixes issue for some cases where negative quaternions not seen in training give bad results in inference.

diff --git a/Engine/Plugins/Experimental/MLDeformer/Source/MLDeformer/Private/MLDeformerInputInfo.cpp b/Engine/Plugins/Experimental/MLDeformer/Source/MLDeformer/Private/MLDeformerInputInfo.cpp
--- a/Engine/Plugins/Experimental/MLDeformer/Source/MLDeformer/Private/MLDeformerInputInfo.cpp
+++ b/Engine/Plugins/Experimental/MLDeformer/Source/MLDeformer/Private/MLDeformerInputInfo.cpp
@@ -248,10 +248,11 @@ void FMLDeformerInputInfo::ExtractBoneRotations(USkeletalMeshComponent* SkelMesh
 		{
 			const int32 Offset = Index * 4;
 			const FQuat Rotation = BoneTransforms[SkelMeshBoneIndex].GetRotation();
-			OutRotations[Offset] = Rotation.X;
-			OutRotations[Offset+1] = Rotation.Y;
-			OutRotations[Offset+2] = Rotation.Z;
-			OutRotations[Offset+3] = Rotation.W;
+			const float QuatSign = Rotation.W < 0.f ? -1.f : 1.f;
+			OutRotations[Offset] = Rotation.X * QuatSign;
+			OutRotations[Offset+1] = Rotation.Y * QuatSign;
+			OutRotations[Offset+2] = Rotation.Z * QuatSign;
+			OutRotations[Offset+3] = Rotation.W * QuatSign;
 		}
 		else
 		{
diff --git a/Engine/Plugins/Experimental/MLDeformer/Source/MLDeformer/Private/MLDeformerInstance.cpp b/Engine/Plugins/Experimental/MLDeformer/Source/MLDeformer/Private/MLDeformerInstance.cpp
--- a/Engine/Plugins/Experimental/MLDeformer/Source/MLDeformer/Private/MLDeformerInstance.cpp
+++ b/Engine/Plugins/Experimental/MLDeformer/Source/MLDeformer/Private/MLDeformerInstance.cpp
@@ -112,10 +112,11 @@ int64 FMLDeformerInstance::SetBoneTransforms(float* OutputBuffer, int64 OutputBu
 		if (SkelMeshBoneIndex != INDEX_NONE)
 		{
 			const FQuat& Rotation = BoneTransforms[SkelMeshBoneIndex].GetRotation();
-			OutputBuffer[Index++] = Rotation.X;
-			OutputBuffer[Index++] = Rotation.Y;
-			OutputBuffer[Index++] = Rotation.Z;
-			OutputBuffer[Index++] = Rotation.W;
+			const float QuatSign = Rotation.W < 0.f ? -1.f : 1.f;
+			OutputBuffer[Index++] = Rotation.X * QuatSign;
+			OutputBuffer[Index++] = Rotation.Y * QuatSign;
+			OutputBuffer[Index++] = Rotation.Z * QuatSign;
+			OutputBuffer[Index++] = Rotation.W * QuatSign;
 		}
 		else
 		{
