MLDeformer: Use new NNI API
Fixes issue with stale DirectML bindings.
Supports multiple MLDeformer components using the same NeuralNetwork.

diff --git a/Engine/Plugins/Experimental/MLDeformer/Source/MLDeformer/Private/MLDeformerComponent.cpp b/Engine/Plugins/Experimental/MLDeformer/Source/MLDeformer/Private/MLDeformerComponent.cpp
--- a/Engine/Plugins/Experimental/MLDeformer/Source/MLDeformer/Private/MLDeformerComponent.cpp
+++ b/Engine/Plugins/Experimental/MLDeformer/Source/MLDeformer/Private/MLDeformerComponent.cpp
@@ -36,6 +36,11 @@ void UMLDeformerComponent::Activate(bool bReset)
 	SetupComponent(DeformerAsset, SkelMeshComponent);
 }
 
+void UMLDeformerComponent::Deactivate()
+{
+	DeformerInstance.Release();
+}
+
 void UMLDeformerComponent::TickComponent(float DeltaTime, enum ELevelTick TickType, FActorComponentTickFunction* ThisTickFunction)
 {
 	if (TickType != ELevelTick::LEVELTICK_PauseTick)
diff --git a/Engine/Plugins/Experimental/MLDeformer/Source/MLDeformer/Private/MLDeformerDataInterface.cpp b/Engine/Plugins/Experimental/MLDeformer/Source/MLDeformer/Private/MLDeformerDataInterface.cpp
--- a/Engine/Plugins/Experimental/MLDeformer/Source/MLDeformer/Private/MLDeformerDataInterface.cpp
+++ b/Engine/Plugins/Experimental/MLDeformer/Source/MLDeformer/Private/MLDeformerDataInterface.cpp
@@ -98,7 +98,8 @@ bool UMLDeformerDataProvider::IsValid() const
 		DeformerComponent->GetDeformerAsset() != nullptr &&
 		DeformerComponent->GetDeformerAsset()->GetVertexMapBuffer().ShaderResourceViewRHI != nullptr &&
 		DeformerComponent->GetDeformerAsset()->GetInferenceNeuralNetwork() != nullptr &&
-		DeformerComponent->GetDeformerAsset()->GetInferenceNeuralNetwork()->IsLoaded();
+		DeformerComponent->GetDeformerAsset()->GetInferenceNeuralNetwork()->IsLoaded() &&
+		DeformerComponent->GetDeformerInstance().GetNeuralNetworkInferenceHandle() != -1;
 }
 
 FComputeDataProviderRenderProxy* UMLDeformerDataProvider::GetRenderProxy()
@@ -112,8 +113,9 @@ FMLDeformerDataProviderProxy::FMLDeformerDataProviderProxy(USkeletalMeshComponen
 
 	const UMLDeformerAsset* DeformerAsset = DeformerComponent->GetDeformerAsset();
 	NeuralNetwork = DeformerAsset->GetInferenceNeuralNetwork();
-	bCanRunNeuralNet = NeuralNetwork ? (NeuralNetwork->GetInputTensor().Num() == static_cast<int64>(DeformerAsset->GetInputInfo().CalcNumNeuralNetInputs())) : false;
-	
+	NeuralNetworkInferenceHandle = DeformerComponent->GetDeformerInstance().GetNeuralNetworkInferenceHandle();
+	bCanRunNeuralNet = NeuralNetwork ? (NeuralNetwork->GetInputTensorForContext(NeuralNetworkInferenceHandle).Num() == static_cast<int64>(DeformerAsset->GetInputInfo().CalcNumNeuralNetInputs())) : false;
+
 	VertexMapBufferSRV = DeformerAsset->GetVertexMapBuffer().ShaderResourceViewRHI;
 	VertexDeltaScale = DeformerAsset->GetVertexDeltaScale();
 	VertexDeltaMean = DeformerAsset->GetVertexDeltaMean();
@@ -126,7 +128,7 @@ void FMLDeformerDataProviderProxy::AllocateResources(FRDGBuilder& GraphBuilder)
 	if (bCanRunNeuralNet)
 	{
 		check(NeuralNetwork != nullptr);
-		Buffer = GraphBuilder.RegisterExternalBuffer(NeuralNetwork->GetOutputTensor().GetPooledBuffer());
+		Buffer = GraphBuilder.RegisterExternalBuffer(NeuralNetwork->GetOutputTensorForContext(NeuralNetworkInferenceHandle).GetPooledBuffer());
 	}
 	else
 	{
diff --git a/Engine/Plugins/Experimental/MLDeformer/Source/MLDeformer/Private/MLDeformerInstance.cpp b/Engine/Plugins/Experimental/MLDeformer/Source/MLDeformer/Private/MLDeformerInstance.cpp
--- a/Engine/Plugins/Experimental/MLDeformer/Source/MLDeformer/Private/MLDeformerInstance.cpp
+++ b/Engine/Plugins/Experimental/MLDeformer/Source/MLDeformer/Private/MLDeformerInstance.cpp
@@ -42,6 +42,16 @@ void FMLDeformerInstance::Init(UMLDeformerAsset* Asset, USkeletalMeshComponent*
 	UpdateCompatibilityStatus();
 }
 
+void FMLDeformerInstance::Release()
+{
+	UNeuralNetwork* NeuralNetwork = DeformerAsset != nullptr ? DeformerAsset->GetInferenceNeuralNetwork() : nullptr;
+	if (NeuralNetwork != nullptr && NeuralNetworkInferenceHandle != -1)
+	{
+		NeuralNetwork->DestroyInferenceContext(NeuralNetworkInferenceHandle);
+		NeuralNetworkInferenceHandle = -1;
+	}
+}
+
 void FMLDeformerInstance::UpdateCompatibilityStatus()
 {
 	bIsCompatible = SkeletalMeshComponent->SkeletalMesh && CheckCompatibility(SkeletalMeshComponent, true).IsEmpty();
@@ -196,8 +206,18 @@ void FMLDeformerInstance::Update()
 	check(NeuralNetwork->GetDeviceType() == ENeuralDeviceType::GPU);
 	check(NeuralNetwork->GetOutputDeviceType() == ENeuralDeviceType::GPU);
 
+	// Allocate an inference context if none has already been allocated.
+	if (NeuralNetworkInferenceHandle == -1)
+	{
+		NeuralNetworkInferenceHandle = NeuralNetwork->CreateInferenceContext();
+		if (NeuralNetworkInferenceHandle == -1)
+		{
+			return;
+		}
+	}
+
 	// If the neural network expects a different number of inputs, do nothing.
-	const int64 NumNeuralNetInputs = NeuralNetwork->GetInputTensor().Num();
+	const int64 NumNeuralNetInputs = NeuralNetwork->GetInputTensorForContext(NeuralNetworkInferenceHandle).Num();
 	const int64 NumDeformerAssetInputs = static_cast<int64>(DeformerAsset->GetInputInfo().CalcNumNeuralNetInputs());
 	if (NumNeuralNetInputs != NumDeformerAssetInputs)
 	{
@@ -205,13 +225,15 @@ void FMLDeformerInstance::Update()
 	}
 
 	// Update and write the input values directly into the input tensor.
-	float* InputDataPointer = (float*)NeuralNetwork->GetInputDataPointerMutable();
+	float* InputDataPointer = (float*)NeuralNetwork->GetInputDataPointerMutableForContext(NeuralNetworkInferenceHandle);
 	SetNeuralNetworkInputValues(InputDataPointer, NumNeuralNetInputs);
 
-	// Output deltas will be available on GPU for the DeformerGraph via UMLDeformerDataProvider.
-	// So this does not actually modify our mesh directly. It just outputs the deltas, which are then used inside a deformer graph later on, accessible
-	// through the data provider class mentioned above.
-	// We pass ENeuralDeviceType::CPU as parameter because the inputs come from the CPU.
-	// We could later switch to Asynchronous processing if we want to run this in a background thread.
-	NeuralNetwork->Run();
+	ENQUEUE_RENDER_COMMAND(RunNeuralNetwork)(
+		[NeuralNetwork, Handle = NeuralNetworkInferenceHandle](FRHICommandListImmediate& RHICmdList)
+	{
+			// Output deltas will be available on GPU for DeformerGraph via UMLDeformerDataProvider.
+			FRDGBuilder GraphBuilder(RHICmdList);
+			NeuralNetwork->Run(GraphBuilder, Handle);
+			GraphBuilder.Execute();
+	});
 }
diff --git a/Engine/Plugins/Experimental/MLDeformer/Source/MLDeformer/Public/MLDeformerComponent.h b/Engine/Plugins/Experimental/MLDeformer/Source/MLDeformer/Public/MLDeformerComponent.h
--- a/Engine/Plugins/Experimental/MLDeformer/Source/MLDeformer/Public/MLDeformerComponent.h
+++ b/Engine/Plugins/Experimental/MLDeformer/Source/MLDeformer/Public/MLDeformerComponent.h
@@ -30,6 +30,7 @@ public:
 
 	//~UActorComponent interface
 	void Activate(bool bReset=false) override;
+	void Deactivate() override;
 	//End of ~UActorComponent interface
 
 	/**
diff --git a/Engine/Plugins/Experimental/MLDeformer/Source/MLDeformer/Public/MLDeformerDataInterface.h b/Engine/Plugins/Experimental/MLDeformer/Source/MLDeformer/Public/MLDeformerDataInterface.h
--- a/Engine/Plugins/Experimental/MLDeformer/Source/MLDeformer/Public/MLDeformerDataInterface.h
+++ b/Engine/Plugins/Experimental/MLDeformer/Source/MLDeformer/Public/MLDeformerDataInterface.h
@@ -72,6 +72,7 @@ private:
 	FSkeletalMeshObject* SkeletalMeshObject;
 	UNeuralNetwork* NeuralNetwork = nullptr;
 	bool bCanRunNeuralNet = false;
+	int32 NeuralNetworkInferenceHandle = -1;
 	FRHIShaderResourceView* VertexMapBufferSRV = nullptr;
 	FRDGBuffer* Buffer = nullptr;
 	FRDGBufferSRV* BufferSRV = nullptr;
diff --git a/Engine/Plugins/Experimental/MLDeformer/Source/MLDeformer/Public/MLDeformerInstance.h b/Engine/Plugins/Experimental/MLDeformer/Source/MLDeformer/Public/MLDeformerInstance.h
--- a/Engine/Plugins/Experimental/MLDeformer/Source/MLDeformer/Public/MLDeformerInstance.h
+++ b/Engine/Plugins/Experimental/MLDeformer/Source/MLDeformer/Public/MLDeformerInstance.h
@@ -25,6 +25,11 @@ public:
 	 */
 	void Init(UMLDeformerAsset* Asset, USkeletalMeshComponent* SkelMeshComponent);
 
+	/**
+	 * Release the resources for this instance.
+	 */
+	void Release();
+
 	/** 
 	 * Update the deformer instance. 
 	 * This updates the neural net input values and runs the actual neural network inference to calculate the output deltas.
@@ -49,6 +54,12 @@ public:
 	 */
 	const FString& GetCompatibilityErrorText() const { return ErrorText; }
 
+	/**
+	 * Get the neural network inference handle.
+	 * @return Returns the handle. This is -1 for an invalid handle.
+	 */
+	int32 GetNeuralNetworkInferenceHandle() const { return NeuralNetworkInferenceHandle; }
+
 	/**
 	 * Get the skeletal mesh component we're working with.
 	 * @return A pointer to the skeletal mesh component.
@@ -99,6 +110,9 @@ private:
 	/** The ML Deformer asset we're an instance of. */
 	TObjectPtr<UMLDeformerAsset> DeformerAsset = nullptr;
 
+	/** Inference handle for neural network inference */
+	int32 NeuralNetworkInferenceHandle = -1;
+
 	/** The skeletal mesh component. */
 	TObjectPtr<USkeletalMeshComponent> SkeletalMeshComponent = nullptr;
 
