Quick fix for crash when NeuralNetwork is swapped out in the Asset.
For example this happens after training.
When this happens we need to recreate any inference exectution contexts.
Might revisit this with a delegate based solution soon.

diff --git a/Engine/Plugins/Experimental/MLDeformer/Source/MLDeformer/Private/MLDeformerInstance.cpp b/Engine/Plugins/Experimental/MLDeformer/Source/MLDeformer/Private/MLDeformerInstance.cpp
--- a/Engine/Plugins/Experimental/MLDeformer/Source/MLDeformer/Private/MLDeformerInstance.cpp
+++ b/Engine/Plugins/Experimental/MLDeformer/Source/MLDeformer/Private/MLDeformerInstance.cpp
@@ -206,14 +206,23 @@ void FMLDeformerInstance::Update()
 	check(NeuralNetwork->GetDeviceType() == ENeuralDeviceType::GPU);
 	check(NeuralNetwork->GetOutputDeviceType() == ENeuralDeviceType::GPU);
 
+	// Check for NeuralNetwork having been replaced below the hood. Maybe because it was retrained.
+	if (CurrentNeuralNetwork != NeuralNetwork)
+	{
+		CurrentNeuralNetwork = NeuralNetwork;
+		// Handle is no longer valid. 
+		// This potentially leaves a dangling inference context with the old NeuralNetwork but in most/all cases the old NeuralNetwork is being destroyed anyway.
+		NeuralNetworkInferenceHandle = -1;
+	}
+
 	// Allocate an inference context if none has already been allocated.
 	if (NeuralNetworkInferenceHandle == -1)
 	{
 		NeuralNetworkInferenceHandle = NeuralNetwork->CreateInferenceContext();
-		if (NeuralNetworkInferenceHandle == -1)
-		{
-			return;
-		}
+	}
+	if (NeuralNetworkInferenceHandle == -1)
+	{
+		return;
 	}
 
 	// If the neural network expects a different number of inputs, do nothing.
diff --git a/Engine/Plugins/Experimental/MLDeformer/Source/MLDeformer/Public/MLDeformerInstance.h b/Engine/Plugins/Experimental/MLDeformer/Source/MLDeformer/Public/MLDeformerInstance.h
--- a/Engine/Plugins/Experimental/MLDeformer/Source/MLDeformer/Public/MLDeformerInstance.h
+++ b/Engine/Plugins/Experimental/MLDeformer/Source/MLDeformer/Public/MLDeformerInstance.h
@@ -10,6 +10,7 @@
 
 // Forward declarations.
 class UMLDeformerAsset;
+class UNeuralNetwork;
 class USkeletalMeshComponent;
 
 /**
@@ -110,6 +111,9 @@ private:
 	/** The ML Deformer asset we're an instance of. */
 	TObjectPtr<UMLDeformerAsset> DeformerAsset = nullptr;
 
+	/** Cache most recent NeuralNetwork from the DeformerAsset. This can change after training. */
+	UNeuralNetwork* CurrentNeuralNetwork = nullptr;
+
 	/** Inference handle for neural network inference */
 	int32 NeuralNetworkInferenceHandle = -1;
 
