MLDeformer: Add delegate to respond to NeuralNetwork being set on an MLDeformerAsset.

diff --git a/Engine/Plugins/Experimental/MLDeformer/Source/MLDeformer/Private/MLDeformerAsset.cpp b/Engine/Plugins/Experimental/MLDeformer/Source/MLDeformer/Private/MLDeformerAsset.cpp
--- a/Engine/Plugins/Experimental/MLDeformer/Source/MLDeformer/Private/MLDeformerAsset.cpp
+++ b/Engine/Plugins/Experimental/MLDeformer/Source/MLDeformer/Private/MLDeformerAsset.cpp
@@ -309,6 +309,12 @@ bool UMLDeformerAsset::IsReadyForFinishDestroy()
 	return Super::IsReadyForFinishDestroy() && RenderResourceDestroyFence.IsFenceComplete();
 }
 
+void UMLDeformerAsset::SetInferenceNeuralNetwork(UNeuralNetwork* InNeuralNetwork)
+{
+	NeuralNetworkModifyDelegate.Broadcast();
+	NeuralNetwork = InNeuralNetwork;
+}
+
 #if WITH_EDITOR
 void UMLDeformerAsset::InitVertexMap()
 {
diff --git a/Engine/Plugins/Experimental/MLDeformer/Source/MLDeformer/Private/MLDeformerComponent.cpp b/Engine/Plugins/Experimental/MLDeformer/Source/MLDeformer/Private/MLDeformerComponent.cpp
--- a/Engine/Plugins/Experimental/MLDeformer/Source/MLDeformer/Private/MLDeformerComponent.cpp
+++ b/Engine/Plugins/Experimental/MLDeformer/Source/MLDeformer/Private/MLDeformerComponent.cpp
@@ -17,10 +17,36 @@ UMLDeformerComponent::UMLDeformerComponent(const FObjectInitializer& ObjectIniti
 
 void UMLDeformerComponent::SetupComponent(UMLDeformerAsset* InDeformerAsset, USkeletalMeshComponent* InSkelMeshComponent)
 {
+	RemoveNeuralNetworkModifyDelegate();
+
 	DeformerAsset = InDeformerAsset;
 	SkelMeshComponent = InSkelMeshComponent;
 	AddTickPrerequisiteComponent(SkelMeshComponent);
 	DeformerInstance.Init(InDeformerAsset, InSkelMeshComponent);
+
+	AddNeuralNetworkModifyDelegate();
+}
+
+void UMLDeformerComponent::AddNeuralNetworkModifyDelegate()
+{
+	if (DeformerAsset != nullptr)
+	{
+		NeuralNetworkModifyDelegateHandle = DeformerAsset->NeuralNetworkModifyDelegate.AddLambda(([this]()
+		{
+			DeformerInstance.Release();
+			DeformerInstance.Init(DeformerAsset, SkelMeshComponent);
+		}));
+	}
+}
+
+void UMLDeformerComponent::RemoveNeuralNetworkModifyDelegate()
+{
+	if (DeformerAsset != nullptr && NeuralNetworkModifyDelegateHandle != FDelegateHandle())
+	{
+		DeformerAsset->NeuralNetworkModifyDelegate.Remove(NeuralNetworkModifyDelegateHandle);
+	}
+	
+	NeuralNetworkModifyDelegateHandle = FDelegateHandle();
 }
 
 void UMLDeformerComponent::Activate(bool bReset)
@@ -38,6 +64,7 @@ void UMLDeformerComponent::Activate(bool bReset)
 
 void UMLDeformerComponent::Deactivate()
 {
+	RemoveNeuralNetworkModifyDelegate();
 	DeformerInstance.Release();
 }
 
diff --git a/Engine/Plugins/Experimental/MLDeformer/Source/MLDeformer/Private/MLDeformerInstance.cpp b/Engine/Plugins/Experimental/MLDeformer/Source/MLDeformer/Private/MLDeformerInstance.cpp
--- a/Engine/Plugins/Experimental/MLDeformer/Source/MLDeformer/Private/MLDeformerInstance.cpp
+++ b/Engine/Plugins/Experimental/MLDeformer/Source/MLDeformer/Private/MLDeformerInstance.cpp
@@ -206,15 +206,6 @@ void FMLDeformerInstance::Update()
 	check(NeuralNetwork->GetDeviceType() == ENeuralDeviceType::GPU);
 	check(NeuralNetwork->GetOutputDeviceType() == ENeuralDeviceType::GPU);
 
-	// Check for NeuralNetwork having been replaced below the hood. Maybe because it was retrained.
-	if (CurrentNeuralNetwork != NeuralNetwork)
-	{
-		CurrentNeuralNetwork = NeuralNetwork;
-		// Handle is no longer valid. 
-		// This potentially leaves a dangling inference context with the old NeuralNetwork but in most/all cases the old NeuralNetwork is being destroyed anyway.
-		NeuralNetworkInferenceHandle = -1;
-	}
-
 	// Allocate an inference context if none has already been allocated.
 	if (NeuralNetworkInferenceHandle == -1)
 	{
diff --git a/Engine/Plugins/Experimental/MLDeformer/Source/MLDeformer/Public/MLDeformerAsset.h b/Engine/Plugins/Experimental/MLDeformer/Source/MLDeformer/Public/MLDeformerAsset.h
--- a/Engine/Plugins/Experimental/MLDeformer/Source/MLDeformer/Public/MLDeformerAsset.h
+++ b/Engine/Plugins/Experimental/MLDeformer/Source/MLDeformer/Public/MLDeformerAsset.h
@@ -159,6 +159,7 @@ public:
 	void InitGPUData();
 
 	const TArray<int32>& GetVertexMap() const { return VertexMap; }
+	void SetInferenceNeuralNetwork(UNeuralNetwork* InNeuralNetwork);
 	UNeuralNetwork* GetInferenceNeuralNetwork() const { return NeuralNetwork.Get(); }
 	const FVertexMapBuffer& GetVertexMapBuffer() const { return VertexMapBuffer; }
 	const FVector3f& GetVertexDeltaMean() const { return VertexDeltaMean; }
@@ -199,11 +200,16 @@ public:
 	float GetNoiseAmount() const { return NoiseAmount; }
 #endif
 
-public:
+protected:
 	/** The neural network to use for inference. */
 	UPROPERTY()
 	TObjectPtr<UNeuralNetwork> NeuralNetwork = nullptr;
 
+public:
+	/** Delegate that will be called immediately before the NeuralNetwork is changed. */
+	DECLARE_MULTICAST_DELEGATE(FNeuralNetworkModifyDelegate);
+	FNeuralNetworkModifyDelegate NeuralNetworkModifyDelegate;
+
 	/** This is an index per vertex in the mesh, indicating the imported vertex number from the source asset. */
 	UPROPERTY()
 	TArray<int32> VertexMap;
diff --git a/Engine/Plugins/Experimental/MLDeformer/Source/MLDeformer/Public/MLDeformerComponent.h b/Engine/Plugins/Experimental/MLDeformer/Source/MLDeformer/Public/MLDeformerComponent.h
--- a/Engine/Plugins/Experimental/MLDeformer/Source/MLDeformer/Public/MLDeformerComponent.h
+++ b/Engine/Plugins/Experimental/MLDeformer/Source/MLDeformer/Public/MLDeformerComponent.h
@@ -65,6 +65,14 @@ private:
 	/** The mesh deformer instance. This is used to perform the runtime updates and run the inference. */
 	FMLDeformerInstance DeformerInstance;
 
+	/** DelegateHandle for NeuralNetwork modification. */
+	FDelegateHandle NeuralNetworkModifyDelegateHandle;
+
+	/** Bind to the DeformerAsset NeuralNetworkModifyDelegate. */
+	void AddNeuralNetworkModifyDelegate();
+	/** Unbind from the DeformerAsset NeuralNetworkModifyDelegate. */
+	void RemoveNeuralNetworkModifyDelegate();
+
 	// BEGIN AActorComponent interface
 	/**
 	 * Tick component, i.e., run MLDeformer inference.
diff --git a/Engine/Plugins/Experimental/MLDeformer/Source/MLDeformer/Public/MLDeformerInstance.h b/Engine/Plugins/Experimental/MLDeformer/Source/MLDeformer/Public/MLDeformerInstance.h
--- a/Engine/Plugins/Experimental/MLDeformer/Source/MLDeformer/Public/MLDeformerInstance.h
+++ b/Engine/Plugins/Experimental/MLDeformer/Source/MLDeformer/Public/MLDeformerInstance.h
@@ -111,9 +111,6 @@ private:
 	/** The ML Deformer asset we're an instance of. */
 	TObjectPtr<UMLDeformerAsset> DeformerAsset = nullptr;
 
-	/** Cache most recent NeuralNetwork from the DeformerAsset. This can change after training. */
-	UNeuralNetwork* CurrentNeuralNetwork = nullptr;
-
 	/** Inference handle for neural network inference */
 	int32 NeuralNetworkInferenceHandle = -1;
 
diff --git a/Engine/Plugins/Experimental/MLDeformer/Source/MLDeformerEditor/Private/MLDeformerEditorToolkit.cpp b/Engine/Plugins/Experimental/MLDeformer/Source/MLDeformerEditor/Private/MLDeformerEditorToolkit.cpp
--- a/Engine/Plugins/Experimental/MLDeformer/Source/MLDeformerEditor/Private/MLDeformerEditorToolkit.cpp
+++ b/Engine/Plugins/Experimental/MLDeformer/Source/MLDeformerEditor/Private/MLDeformerEditorToolkit.cpp
@@ -375,8 +375,9 @@ bool FMLDeformerEditorToolkit::TryLoadOnnxFile() const
 		UNeuralNetwork* Network = NewObject<UNeuralNetwork>(EditorData->GetDeformerAsset(), UNeuralNetwork::StaticClass());		
 		if (Network->Load(OnnxFile))
 		{
-			EditorData->GetDeformerAsset()->NeuralNetwork = Network;
-			EditorData->GetDeformerAsset()->NeuralNetwork->SetDeviceType(/*DeviceType*/ENeuralDeviceType::GPU, /*InputDeviceType*/ENeuralDeviceType::CPU, /*OutputDeviceType*/ENeuralDeviceType::GPU);
+			Network->SetDeviceType(/*DeviceType*/ENeuralDeviceType::GPU, /*InputDeviceType*/ENeuralDeviceType::CPU, /*OutputDeviceType*/ENeuralDeviceType::GPU);
+			
+			EditorData->GetDeformerAsset()->SetInferenceNeuralNetwork(Network);
 
 			// Recreate the data providers.
 			// This is needed for the neural network GPU buffers to be valid.
