[MLDeformer]
- Fixed a bug where you train an already trained network but now with lower amount of training samples by modifying Max Training samples. If you would cancel and choose not to use the partially trained network, the deformation would mess up.
- Fixed a bug where 'Bones and Curves' caused a training error in the python script.

diff --git a/Engine/Plugins/Experimental/MLDeformer/Source/MLDeformerEditor/Private/MLDeformerEditorToolkit.cpp b/Engine/Plugins/Experimental/MLDeformer/Source/MLDeformerEditor/Private/MLDeformerEditorToolkit.cpp
--- a/Engine/Plugins/Experimental/MLDeformer/Source/MLDeformerEditor/Private/MLDeformerEditorToolkit.cpp
+++ b/Engine/Plugins/Experimental/MLDeformer/Source/MLDeformerEditor/Private/MLDeformerEditorToolkit.cpp
@@ -171,6 +171,10 @@ void FMLDeformerEditorToolkit::FillToolbar(FToolBarBuilder& ToolbarBuilder)
 					{
 						ShowNotification(LOCTEXT("StartTraining", "Starting training process"), SNotificationItem::ECompletionState::CS_Pending, true);
 
+						// Backup the mean and scale.
+						EditorData->VertexDeltaMeanBackup = DeformerAsset->GetVertexDeltaMean();
+						EditorData->VertexDeltaScaleBackup = DeformerAsset->GetVertexDeltaScale();
+
 						// Change the interpolation type for the training sequence to step.
 						DeformerAsset->GetAnimSequence()->Interpolation = EAnimInterpolationType::Step;
 
@@ -213,6 +217,7 @@ void FMLDeformerEditorToolkit::FillToolbar(FToolBarBuilder& ToolbarBuilder)
 							// If we aborted we still have normalized the data. Only when we have AbortedCantUse then we canceled the normalization process.
 							EditorData->bIsVertexDeltaNormalized = true;
 						}
+
 						EditorData->GetEditorActor(EMLDeformerEditorActorIndex::DeformedTest).MLDeformerComponent->SetupComponent(DeformerAsset, SkelMeshComponent);
 						if (bMarkDirty)
 						{
@@ -323,6 +328,14 @@ bool FMLDeformerEditorToolkit::HandleTrainingResult(ETrainingResult TrainingResu
 			}
 			else
 			{
+				// Restore the vertex delta mean and scale, as we aborted, and they could have changed when training
+				// on a smaller subset of frames/samples. If we don't do this, the mesh will deform incorrectly.
+				if (TrainingResult == ETrainingResult::Aborted)
+				{
+					EditorData->GetDeformerAsset()->VertexDeltaMean = EditorData->VertexDeltaMeanBackup;
+					EditorData->GetDeformerAsset()->VertexDeltaScale = EditorData->VertexDeltaScaleBackup;
+				}
+
 				ShowNotification(LOCTEXT("TrainingAborted", "Training aborted!"), SNotificationItem::ECompletionState::CS_None, true);
 			}
 		}
diff --git a/Engine/Plugins/Experimental/MLDeformer/Source/MLDeformerEditor/Public/MLDeformerEditorData.h b/Engine/Plugins/Experimental/MLDeformer/Source/MLDeformerEditor/Public/MLDeformerEditorData.h
--- a/Engine/Plugins/Experimental/MLDeformer/Source/MLDeformerEditor/Public/MLDeformerEditorData.h
+++ b/Engine/Plugins/Experimental/MLDeformer/Source/MLDeformerEditor/Public/MLDeformerEditorData.h
@@ -131,6 +131,9 @@ public:
 	TArray<FVector3f> DebugVectors2;
 	TArray<float> VertexDeltas;
 
+	FVector3f VertexDeltaMeanBackup = FVector3f::ZeroVector;
+	FVector3f VertexDeltaScaleBackup = FVector3f::OneVector;
+
 	bool bIsVertexDeltaNormalized = false;
 	bool bIsReadyForTraining = false;
 	int32 CurrentFrame = -1;
