[MLDeformer] Static analyzer fix.

diff --git a/Engine/Plugins/Experimental/MLDeformer/Source/MLDeformerEditor/Private/MLDeformerEditorToolkit.cpp b/Engine/Plugins/Experimental/MLDeformer/Source/MLDeformerEditor/Private/MLDeformerEditorToolkit.cpp
--- a/Engine/Plugins/Experimental/MLDeformer/Source/MLDeformerEditor/Private/MLDeformerEditorToolkit.cpp
+++ b/Engine/Plugins/Experimental/MLDeformer/Source/MLDeformerEditor/Private/MLDeformerEditorToolkit.cpp
@@ -330,11 +330,8 @@ bool FMLDeformerEditorToolkit::HandleTrainingResult(ETrainingResult TrainingResu
 			{
 				// Restore the vertex delta mean and scale, as we aborted, and they could have changed when training
 				// on a smaller subset of frames/samples. If we don't do this, the mesh will deform incorrectly.
-				if (TrainingResult == ETrainingResult::Aborted)
-				{
-					EditorData->GetDeformerAsset()->VertexDeltaMean = EditorData->VertexDeltaMeanBackup;
-					EditorData->GetDeformerAsset()->VertexDeltaScale = EditorData->VertexDeltaScaleBackup;
-				}
+				EditorData->GetDeformerAsset()->VertexDeltaMean = EditorData->VertexDeltaMeanBackup;
+				EditorData->GetDeformerAsset()->VertexDeltaScale = EditorData->VertexDeltaScaleBackup;
 
 				ShowNotification(LOCTEXT("TrainingAborted", "Training aborted!"), SNotificationItem::ECompletionState::CS_None, true);
 			}
