[MLDeformer] Fix ground truth heatmaps giving buggy results sometimes. Note: playing the animation while ground truth heatmap is enabled there is still flickering, but that is another issue.

diff --git a/Engine/Plugins/Experimental/MLDeformer/Source/MLDeformer/Private/MLDeformerAsset.cpp b/Engine/Plugins/Experimental/MLDeformer/Source/MLDeformer/Private/MLDeformerAsset.cpp
--- a/Engine/Plugins/Experimental/MLDeformer/Source/MLDeformer/Private/MLDeformerAsset.cpp
+++ b/Engine/Plugins/Experimental/MLDeformer/Source/MLDeformer/Private/MLDeformerAsset.cpp
@@ -109,13 +109,10 @@ bool IsPotentialMatch(const FString& TrackName, const FString& MeshName)
 	return (TrackName.Find(MeshName) == 0);
 }
 
-void UMLDeformerAsset::GenerateMeshMappings(UMLDeformerAsset* DeformerAsset, TArray<FMLDeformerMeshMapping>& OutMeshMappings, TArray<FString>& OutFailedImportedMeshNames)
+void UMLDeformerAsset::GenerateMeshMappings(USkeletalMesh* SkelMesh, UGeometryCache* GeomCache, TArray<FMLDeformerMeshMapping>& OutMeshMappings, TArray<FString>& OutFailedImportedMeshNames)
 {
 	OutMeshMappings.Empty();
 	OutFailedImportedMeshNames.Empty();
-
-	USkeletalMesh* SkelMesh = DeformerAsset->GetSkeletalMesh();
-	UGeometryCache* GeomCache = DeformerAsset->GetGeometryCache();
 	if (SkelMesh == nullptr || GeomCache == nullptr)
 	{
 		return;
@@ -561,10 +558,9 @@ FText UMLDeformerAsset::GetMeshMappingErrorText() const
 	if (GeometryCache && SkeletalMesh)
 	{
 		// Check for failed mesh mappings.
-		UMLDeformerAsset* Asset = const_cast<UMLDeformerAsset*>(this);
 		TArray<FMLDeformerMeshMapping> MeshMappings;
 		TArray<FString> FailedNames;
-		UMLDeformerAsset::GenerateMeshMappings(Asset, MeshMappings, FailedNames);
+		UMLDeformerAsset::GenerateMeshMappings(SkeletalMesh, GeometryCache, MeshMappings, FailedNames);
 
 		// List all mesh names that have issues.
 		FString ErrorString;
diff --git a/Engine/Plugins/Experimental/MLDeformer/Source/MLDeformer/Private/MLDeformerDebugDataInterface.cpp b/Engine/Plugins/Experimental/MLDeformer/Source/MLDeformer/Private/MLDeformerDebugDataInterface.cpp
--- a/Engine/Plugins/Experimental/MLDeformer/Source/MLDeformer/Private/MLDeformerDebugDataInterface.cpp
+++ b/Engine/Plugins/Experimental/MLDeformer/Source/MLDeformer/Private/MLDeformerDebugDataInterface.cpp
@@ -113,6 +113,30 @@ void UMLDeformerDebugDataInterface::GetSourceTypes(TArray<UClass*>& OutSourceTyp
 	OutSourceTypes.Add(UMLDeformerComponent::StaticClass());
 }
 
+#if WITH_EDITORONLY_DATA
+// Return a GeometryCache only if it matches the current test sequence.
+UGeometryCache* GetActiveGeometryCache(UMLDeformerAsset const* DeformerAsset)
+{
+	UGeometryCache* GeometryCache = DeformerAsset->GetVizSettings()->GetGroundTruth();
+	GeometryCache = (GeometryCache != nullptr) ? GeometryCache : const_cast<UGeometryCache*>(DeformerAsset->GetGeometryCache());
+
+	UAnimSequence const* AnimSequence = DeformerAsset->GetVizSettings()->GetTestAnimSequence();
+
+	if (GeometryCache != nullptr && AnimSequence != nullptr)
+	{
+		const float AnimSeqDuration = AnimSequence->GetPlayLength();
+		const float GeomCacheDuration = GeometryCache->CalculateDuration();
+		if (FMath::Abs(AnimSeqDuration - GeomCacheDuration) < 0.001f)
+		{
+			return GeometryCache;
+		}
+	}
+
+	return nullptr;
+}
+#endif
+
+
 UComputeDataProvider* UMLDeformerDebugDataInterface::CreateDataProvider(TArrayView<TObjectPtr<UObject>> InSourceObjects, uint64 InInputMask, uint64 InOutputMask) const
 {
 	UMLDeformerDebugDataProvider* Provider = NewObject<UMLDeformerDebugDataProvider>();
@@ -122,11 +146,13 @@ UComputeDataProvider* UMLDeformerDebugDataInterface::CreateDataProvider(TArrayVi
 		Provider->DeformerComponent = Cast<UMLDeformerComponent>(InSourceObjects[1]);
 	}
 #if WITH_EDITORONLY_DATA
-	UMLDeformerAsset* DeformerAsset = Provider->DeformerComponent != nullptr ? Provider->DeformerComponent->GetDeformerAsset() : nullptr;
+	UMLDeformerAsset* DeformerAsset = (Provider->DeformerComponent != nullptr) ? Provider->DeformerComponent->GetDeformerAsset() : nullptr;
 	if (DeformerAsset != nullptr)
 	{
 		TArray<FString> FailedImportedMeshNames;
-		UMLDeformerAsset::GenerateMeshMappings(DeformerAsset, Provider->MeshMappings, FailedImportedMeshNames);
+		USkeletalMesh* SkelMesh = Provider->SkeletalMeshComponent->SkeletalMesh;
+		UGeometryCache* GeomCache = GetActiveGeometryCache(DeformerAsset);
+		UMLDeformerAsset::GenerateMeshMappings(SkelMesh, GeomCache, Provider->MeshMappings, FailedImportedMeshNames);
 	}
 #endif
 	return Provider;
@@ -157,27 +183,6 @@ FComputeDataProviderRenderProxy* UMLDeformerDebugDataProvider::GetRenderProxy()
 
 #if WITH_EDITORONLY_DATA
 
-// Return a GeometryCache only if it matches the current test sequence.
-UGeometryCache const* GetActiveGeometryCache(UMLDeformerAsset const* DeformerAsset)
-{
-	UGeometryCache const* GeometryCache = DeformerAsset->GetVizSettings()->GetGroundTruth();
-	GeometryCache = GeometryCache != nullptr ? GeometryCache : DeformerAsset->GetGeometryCache();
-
-	UAnimSequence const* AnimSequence = DeformerAsset->GetVizSettings()->GetTestAnimSequence();
-
-	if (GeometryCache != nullptr && AnimSequence != nullptr)
-	{
-		const float AnimSeqDuration = AnimSequence->GetPlayLength();
-		const float GeomCacheDuration = GeometryCache->CalculateDuration();
-		if (FMath::Abs(AnimSeqDuration - GeomCacheDuration) < 0.001f)
-		{
-			return GeometryCache;
-		}
-	}
-
-	return nullptr;
-}
-
 // Fill the ground truth positions from the geometry cache.
 void GetGroundTruthPositions(
 	int32 LODIndex,
diff --git a/Engine/Plugins/Experimental/MLDeformer/Source/MLDeformer/Public/MLDeformerAsset.h b/Engine/Plugins/Experimental/MLDeformer/Source/MLDeformer/Public/MLDeformerAsset.h
--- a/Engine/Plugins/Experimental/MLDeformer/Source/MLDeformer/Public/MLDeformerAsset.h
+++ b/Engine/Plugins/Experimental/MLDeformer/Source/MLDeformer/Public/MLDeformerAsset.h
@@ -145,7 +145,7 @@ public:
 
 	static int32 ExtractNumImportedGeomCacheVertices(UGeometryCache* GeomCache);
 	static int32 ExtractNumImportedSkinnedVertices(USkeletalMesh* SkeletalMesh);
-	static void GenerateMeshMappings(UMLDeformerAsset* DeformerAsset, TArray<FMLDeformerMeshMapping>& OutMeshMappings, TArray<FString>& OutFailedImportedMeshNames);
+	static void GenerateMeshMappings(USkeletalMesh* SkelMesh, UGeometryCache* GeomCache, TArray<FMLDeformerMeshMapping>& OutMeshMappings, TArray<FString>& OutFailedImportedMeshNames);
 
 	void UpdateCachedNumVertices();
 	bool IsCompatibleWithNeuralNet() const;
diff --git a/Engine/Plugins/Experimental/MLDeformer/Source/MLDeformerEditor/Private/MLDeformerFrameCache.cpp b/Engine/Plugins/Experimental/MLDeformer/Source/MLDeformerEditor/Private/MLDeformerFrameCache.cpp
--- a/Engine/Plugins/Experimental/MLDeformer/Source/MLDeformerEditor/Private/MLDeformerFrameCache.cpp
+++ b/Engine/Plugins/Experimental/MLDeformer/Source/MLDeformerEditor/Private/MLDeformerFrameCache.cpp
@@ -400,7 +400,11 @@ void FMLDeformerSampler::Init(const FInitSettings& InInitSettings)
 	SkelMeshActor->SetRootComponent(SkelMeshComponent);
 	GeomCacheActor->SetRootComponent(GeomCacheComponent);
 
-	UMLDeformerAsset::GenerateMeshMappings(InitSettings.DeformerAsset, MeshMappings, FailedImportedMeshNames);
+	UMLDeformerAsset::GenerateMeshMappings(
+		DeformerAsset->GetSkeletalMesh(), 
+		DeformerAsset->GetGeometryCache(), 
+		MeshMappings, 
+		FailedImportedMeshNames);
 }
 
 // Spawn an actor.
