[MLDeformer] Fixed some loctext duplicate keys.

#rnx

diff --git a/Engine/Plugins/Experimental/MLDeformer/Source/MLDeformerEditor/Private/MLDeformerEditorToolkit.cpp b/Engine/Plugins/Experimental/MLDeformer/Source/MLDeformerEditor/Private/MLDeformerEditorToolkit.cpp
--- a/Engine/Plugins/Experimental/MLDeformer/Source/MLDeformerEditor/Private/MLDeformerEditorToolkit.cpp
+++ b/Engine/Plugins/Experimental/MLDeformer/Source/MLDeformerEditor/Private/MLDeformerEditorToolkit.cpp
@@ -318,7 +318,7 @@ bool FMLDeformerEditorToolkit::HandleTrainingResult(ETrainingResult TrainingResu
 			{
 				if (!TryLoadOnnxFile())
 				{
-					ShowNotification(LOCTEXT("TrainingOnnxLoadFailed", "Training partially completed, but resulting onnx file couldn't be loaded!"), SNotificationItem::ECompletionState::CS_Fail, true);
+					ShowNotification(LOCTEXT("TrainingOnnxLoadFailedPartial", "Training partially completed, but resulting onnx file couldn't be loaded!"), SNotificationItem::ECompletionState::CS_Fail, true);
 				}
 				else
 				{
@@ -350,7 +350,7 @@ bool FMLDeformerEditorToolkit::HandleTrainingResult(ETrainingResult TrainingResu
 		case ETrainingResult::FailOnData:
 		{
 			GEditor->PlayEditorSound(TEXT("/Engine/EditorSounds/Notifications/CompileFailed_Cue.CompileFailed_Cue"));
-			WindowMessage = LOCTEXT("TrainingFailed", "Training failed!\nCheck input parameters or sequence length.");
+			WindowMessage = LOCTEXT("TrainingFailedOnData", "Training failed!\nCheck input parameters or sequence length.");
 		}
 		break;
 
@@ -358,7 +358,7 @@ bool FMLDeformerEditorToolkit::HandleTrainingResult(ETrainingResult TrainingResu
 		case ETrainingResult::FailUnknown:
 		{
 			GEditor->PlayEditorSound(TEXT("/Engine/EditorSounds/Notifications/CompileFailed_Cue.CompileFailed_Cue"));
-			WindowMessage = LOCTEXT("TrainingFailed", "Training failed!\nUnknown error, please check the output log.");
+			WindowMessage = LOCTEXT("TrainingFailedUnknown", "Training failed!\nUnknown error, please check the output log.");
 		}
 		break;
 
