Find MLDeformer component on every update.
In a blueprint the UMLDeformerComponent will be reconstructed at level start or at any data edit which can lead to the DataProvider cached reference being nulled.
Would be better in future to find a way to hook events to rebuild DataProviders instead of doing work every frame.
#rnx

diff --git a/Engine/Plugins/Experimental/MLDeformer/Source/MLDeformer/Private/MLDeformerDataInterface.cpp b/Engine/Plugins/Experimental/MLDeformer/Source/MLDeformer/Private/MLDeformerDataInterface.cpp
--- a/Engine/Plugins/Experimental/MLDeformer/Source/MLDeformer/Private/MLDeformerDataInterface.cpp
+++ b/Engine/Plugins/Experimental/MLDeformer/Source/MLDeformer/Private/MLDeformerDataInterface.cpp
@@ -75,25 +75,28 @@ void UMLDeformerDataInterface::GetHLSL(FString& OutHLSL) const
 void UMLDeformerDataInterface::GetSourceTypes(TArray<UClass*>& OutSourceTypes) const
 {
 	OutSourceTypes.Add(USkeletalMeshComponent::StaticClass());
-	OutSourceTypes.Add(UMLDeformerComponent::StaticClass());
 }
 
 UComputeDataProvider* UMLDeformerDataInterface::CreateDataProvider(TArrayView<TObjectPtr<UObject>> InSourceObjects, uint64 InInputMask, uint64 InOutputMask) const
 {
 	UMLDeformerDataProvider* Provider = NewObject<UMLDeformerDataProvider>();
-	if (InSourceObjects.Num() == 2)
+	if (InSourceObjects.Num() == 1)
 	{
 		Provider->SkeletalMeshComponent = Cast<USkeletalMeshComponent>(InSourceObjects[0]);
-		Provider->DeformerComponent = Cast<UMLDeformerComponent>(InSourceObjects[1]);
 	}
 	return Provider;
 }
 
 bool UMLDeformerDataProvider::IsValid() const
 {
+	if (SkeletalMeshComponent == nullptr || SkeletalMeshComponent->MeshObject == nullptr)
+	{
+		return false;
+	}
+	
+	UMLDeformerComponent* DeformerComponent = SkeletalMeshComponent->GetOwner()->FindComponentByClass<UMLDeformerComponent>();
+
 	return
-		SkeletalMeshComponent != nullptr &&
-		SkeletalMeshComponent->MeshObject != nullptr &&
 		DeformerComponent != nullptr &&
 		DeformerComponent->GetDeformerAsset() != nullptr &&
 		DeformerComponent->GetDeformerAsset()->GetVertexMapBuffer().ShaderResourceViewRHI != nullptr &&
@@ -104,6 +107,8 @@ bool UMLDeformerDataProvider::IsValid() const
 
 FComputeDataProviderRenderProxy* UMLDeformerDataProvider::GetRenderProxy()
 {
+	UMLDeformerComponent* DeformerComponent = SkeletalMeshComponent->GetOwner()->FindComponentByClass<UMLDeformerComponent>();
+
 	return new FMLDeformerDataProviderProxy(SkeletalMeshComponent, DeformerComponent);
 }
 
diff --git a/Engine/Plugins/Experimental/MLDeformer/Source/MLDeformer/Private/MLDeformerDebugDataInterface.cpp b/Engine/Plugins/Experimental/MLDeformer/Source/MLDeformer/Private/MLDeformerDebugDataInterface.cpp
--- a/Engine/Plugins/Experimental/MLDeformer/Source/MLDeformer/Private/MLDeformerDebugDataInterface.cpp
+++ b/Engine/Plugins/Experimental/MLDeformer/Source/MLDeformer/Private/MLDeformerDebugDataInterface.cpp
@@ -143,15 +143,16 @@ UComputeDataProvider* UMLDeformerDebugDataInterface::CreateDataProvider(TArrayVi
 	if (InSourceObjects.Num() == 2)
 	{
 		Provider->SkeletalMeshComponent = Cast<USkeletalMeshComponent>(InSourceObjects[0]);
-		Provider->DeformerComponent = Cast<UMLDeformerComponent>(InSourceObjects[1]);
+		
+		UMLDeformerComponent* DeformerComponent = Cast<UMLDeformerComponent>(InSourceObjects[1]);
+		Provider->DeformerAsset = DeformerComponent != nullptr ? DeformerComponent->GetDeformerAsset() : nullptr;
 	}
 #if WITH_EDITORONLY_DATA
-	UMLDeformerAsset* DeformerAsset = (Provider->DeformerComponent != nullptr) ? Provider->DeformerComponent->GetDeformerAsset() : nullptr;
-	if (DeformerAsset != nullptr)
+	if (Provider->DeformerAsset != nullptr)
 	{
 		TArray<FString> FailedImportedMeshNames;
 		USkeletalMesh* SkelMesh = Provider->SkeletalMeshComponent->SkeletalMesh;
-		UGeometryCache* GeomCache = GetActiveGeometryCache(DeformerAsset);
+		UGeometryCache* GeomCache = GetActiveGeometryCache(Provider->DeformerAsset);
 		UMLDeformerAsset::GenerateMeshMappings(SkelMesh, GeomCache, Provider->MeshMappings, FailedImportedMeshNames);
 	}
 #endif
@@ -164,9 +165,8 @@ bool UMLDeformerDebugDataProvider::IsValid() const
 	return
 		SkeletalMeshComponent != nullptr &&
 		SkeletalMeshComponent->MeshObject != nullptr &&
-		DeformerComponent != nullptr &&
-		DeformerComponent->GetDeformerAsset() != nullptr &&
-		DeformerComponent->GetDeformerAsset()->GetVertexMapBuffer().ShaderResourceViewRHI != nullptr;
+		DeformerAsset != nullptr &&
+		DeformerAsset->GetVertexMapBuffer().ShaderResourceViewRHI != nullptr;
 #else
 	return false; // This data interface is only valid in editor.
 #endif
@@ -175,7 +175,7 @@ bool UMLDeformerDebugDataProvider::IsValid() const
 FComputeDataProviderRenderProxy* UMLDeformerDebugDataProvider::GetRenderProxy()
 {
 #if WITH_EDITORONLY_DATA
-	return new FMLDeformerDebugDataProviderProxy(SkeletalMeshComponent, DeformerComponent, MeshMappings);
+	return new FMLDeformerDebugDataProviderProxy(SkeletalMeshComponent, DeformerAsset, MeshMappings);
 #else
 	return nullptr;
 #endif
@@ -242,11 +242,10 @@ void GetGroundTruthPositions(
 	}
 }
 
-FMLDeformerDebugDataProviderProxy::FMLDeformerDebugDataProviderProxy(USkeletalMeshComponent* SkeletalMeshComponent, UMLDeformerComponent* DeformerComponent, TArray<FMLDeformerMeshMapping> const& MeshMappings)
+FMLDeformerDebugDataProviderProxy::FMLDeformerDebugDataProviderProxy(USkeletalMeshComponent* SkeletalMeshComponent, UMLDeformerAsset* DeformerAsset, TArray<FMLDeformerMeshMapping> const& MeshMappings)
 {
 	SkeletalMeshObject = SkeletalMeshComponent->MeshObject;
-	
-	UMLDeformerAsset const* DeformerAsset = DeformerComponent->GetDeformerAsset();
+
 	VertexMapBufferSRV = DeformerAsset->GetVertexMapBuffer().ShaderResourceViewRHI;
 	HeatMapMode = (int32)DeformerAsset->GetVizSettings()->GetHeatMapMode();
 	HeatMapScale = 1.f / FMath::Max(DeformerAsset->GetVizSettings()->GetHeatMapScale(), 0.00001f);
diff --git a/Engine/Plugins/Experimental/MLDeformer/Source/MLDeformer/Public/MLDeformerDataInterface.h b/Engine/Plugins/Experimental/MLDeformer/Source/MLDeformer/Public/MLDeformerDataInterface.h
--- a/Engine/Plugins/Experimental/MLDeformer/Source/MLDeformer/Public/MLDeformerDataInterface.h
+++ b/Engine/Plugins/Experimental/MLDeformer/Source/MLDeformer/Public/MLDeformerDataInterface.h
@@ -47,8 +47,6 @@ class MLDEFORMER_API UMLDeformerDataProvider : public UComputeDataProvider
 public:
 	UPROPERTY(EditAnywhere, BlueprintReadWrite, Category = Binding)
 	TObjectPtr<USkeletalMeshComponent> SkeletalMeshComponent = nullptr;
-	UPROPERTY(EditAnywhere, BlueprintReadWrite, Category = Binding)
-	TObjectPtr<UMLDeformerComponent> DeformerComponent = nullptr;
 
 	//~ Begin UComputeDataProvider Interface
 	bool IsValid() const override;
diff --git a/Engine/Plugins/Experimental/MLDeformer/Source/MLDeformer/Public/MLDeformerDebugDataInterface.h b/Engine/Plugins/Experimental/MLDeformer/Source/MLDeformer/Public/MLDeformerDebugDataInterface.h
--- a/Engine/Plugins/Experimental/MLDeformer/Source/MLDeformer/Public/MLDeformerDebugDataInterface.h
+++ b/Engine/Plugins/Experimental/MLDeformer/Source/MLDeformer/Public/MLDeformerDebugDataInterface.h
@@ -13,7 +13,7 @@ class FRDGBufferSRV;
 struct FMLDeformerMeshMapping;
 class FRHIShaderResourceView;
 class FSkeletalMeshObject;
-class UMLDeformerComponent;
+class UMLDeformerAsset;
 class USkeletalMeshComponent;
 
 /** 
@@ -51,7 +51,7 @@ public:
 	UPROPERTY(EditAnywhere, BlueprintReadWrite, Category = Binding)
 	TObjectPtr<USkeletalMeshComponent> SkeletalMeshComponent = nullptr;
 	UPROPERTY(EditAnywhere, BlueprintReadWrite, Category = Binding)
-	TObjectPtr<UMLDeformerComponent> DeformerComponent = nullptr;
+	TObjectPtr<UMLDeformerAsset> DeformerAsset = nullptr;
 
 #if WITH_EDITORONLY_DATA
 	/* 
@@ -73,7 +73,7 @@ class FMLDeformerDebugDataProviderProxy : public FComputeDataProviderRenderProxy
 {
 public:
 #if WITH_EDITORONLY_DATA
-	FMLDeformerDebugDataProviderProxy(USkeletalMeshComponent* SkeletalMeshComponent, UMLDeformerComponent* DeformerComponent, TArray<FMLDeformerMeshMapping> const& MeshMappings);
+	FMLDeformerDebugDataProviderProxy(USkeletalMeshComponent* SkeletalMeshComponent, UMLDeformerAsset* DeformerAsset, TArray<FMLDeformerMeshMapping> const& MeshMappings);
 #endif
 
 	//~ Begin FComputeDataProviderRenderProxy Interface
