[MLDeformer] Switched neural network rotations from Quaternions to 2 columns of the 3x3 rotation matrix.

diff --git a/Engine/Plugins/Experimental/MLDeformer/Source/MLDeformer/Private/MLDeformerAsset.cpp b/Engine/Plugins/Experimental/MLDeformer/Source/MLDeformer/Private/MLDeformerAsset.cpp
--- a/Engine/Plugins/Experimental/MLDeformer/Source/MLDeformer/Private/MLDeformerAsset.cpp
+++ b/Engine/Plugins/Experimental/MLDeformer/Source/MLDeformer/Private/MLDeformerAsset.cpp
@@ -733,7 +733,7 @@ FMLDeformerInputInfo UMLDeformerAsset::CreateInputInfo() const
 	}
 	else
 	{
-		UE_LOG(LogMLDeformer, Display, TEXT("Including ALL bones of skeleton in training."));
+		UE_LOG(LogMLDeformer, Display, TEXT("Including ALL curves of skeleton in training."));
 	}
 
 	Result.Init(Settings);
diff --git a/Engine/Plugins/Experimental/MLDeformer/Source/MLDeformer/Private/MLDeformerInputInfo.cpp b/Engine/Plugins/Experimental/MLDeformer/Source/MLDeformer/Private/MLDeformerInputInfo.cpp
--- a/Engine/Plugins/Experimental/MLDeformer/Source/MLDeformer/Private/MLDeformerInputInfo.cpp
+++ b/Engine/Plugins/Experimental/MLDeformer/Source/MLDeformer/Private/MLDeformerInputInfo.cpp
@@ -238,36 +238,29 @@ void FMLDeformerInputInfo::ExtractBoneRotations(USkeletalMeshComponent* SkelMesh
 {
 	const TArray<FTransform>& BoneTransforms = SkelMeshComponent->GetBoneSpaceTransforms();
 	const int32 NumBones = GetNumBones();
-	OutRotations.Reset(NumBones * 4); // xyzw
-	OutRotations.AddUninitialized(NumBones * 4);
+	const int32 NumFloats = NumBones * 6; // 2 Columns of the rotation matrix.
+	OutRotations.Reset(NumFloats);
+	OutRotations.AddUninitialized(NumFloats);
+	int32 Offset = 0;
 	for (int32 Index = 0; Index < NumBones; ++Index)
 	{
 		const FName BoneName = GetBoneName(Index);
 		const int32 SkelMeshBoneIndex = SkelMeshComponent->GetBoneIndex(BoneName);
-		if (SkelMeshBoneIndex != INDEX_NONE)
-		{
-			const int32 Offset = Index * 4;
-			const FQuat Rotation = BoneTransforms[SkelMeshBoneIndex].GetRotation();
-			const float QuatSign = Rotation.W < 0.f ? -1.f : 1.f;
-			OutRotations[Offset] = Rotation.X * QuatSign;
-			OutRotations[Offset+1] = Rotation.Y * QuatSign;
-			OutRotations[Offset+2] = Rotation.Z * QuatSign;
-			OutRotations[Offset+3] = Rotation.W * QuatSign;
-		}
-		else
-		{
-			const int32 Offset = Index * 4;
-			OutRotations[Offset] = 0.0f;	// Set to identity quat.
-			OutRotations[Offset+1] = 0.0f;
-			OutRotations[Offset+2] = 0.0f;
-			OutRotations[Offset+3] = 1.0f;
-		}
+		const FMatrix RotationMatrix = (SkelMeshBoneIndex != INDEX_NONE) ? BoneTransforms[SkelMeshBoneIndex].GetRotation().ToMatrix() : FMatrix::Identity;
+		const FVector X = RotationMatrix.GetColumn(0);
+		const FVector Y = RotationMatrix.GetColumn(1);	
+		OutRotations[Offset++] = X.X;
+		OutRotations[Offset++] = X.Y;
+		OutRotations[Offset++] = X.Z;
+		OutRotations[Offset++] = Y.X;
+		OutRotations[Offset++] = Y.Y;
+		OutRotations[Offset++] = Y.Z;
 	}
 }
 
 int32 FMLDeformerInputInfo::CalcNumNeuralNetInputs() const
 {
 	return 
-		BoneNameStrings.Num() * 4 +	// Four floats per bone.
+		BoneNameStrings.Num() * 6 +	// Six floats per bone.
 		CurveNameStrings.Num();		// One float per curve.
 }
diff --git a/Engine/Plugins/Experimental/MLDeformer/Source/MLDeformer/Private/MLDeformerInstance.cpp b/Engine/Plugins/Experimental/MLDeformer/Source/MLDeformer/Private/MLDeformerInstance.cpp
--- a/Engine/Plugins/Experimental/MLDeformer/Source/MLDeformer/Private/MLDeformerInstance.cpp
+++ b/Engine/Plugins/Experimental/MLDeformer/Source/MLDeformer/Private/MLDeformerInstance.cpp
@@ -115,27 +115,19 @@ int64 FMLDeformerInstance::SetBoneTransforms(float* OutputBuffer, int64 OutputBu
 	const FMLDeformerInputInfo& InputInfo = DeformerAsset->GetInputInfo();
 	const int32 AssetNumBones = InputInfo.GetNumBones();
 	int64 Index = StartIndex;
-	check((Index + AssetNumBones * 4) <= OutputBufferSize); // Make sure we don't write past the OutputBuffer.
+	check((Index + AssetNumBones * 6) <= OutputBufferSize); // Make sure we don't write past the OutputBuffer. (6 because of two columns of the 3x3 rotation matrix)
 	for (int32 BoneIndex = 0; BoneIndex < AssetNumBones; ++BoneIndex)
 	{
 		const int32 SkelMeshBoneIndex = AssetBonesToSkelMeshMappings[BoneIndex];
-		if (SkelMeshBoneIndex != INDEX_NONE)
-		{
-			const FQuat& Rotation = BoneTransforms[SkelMeshBoneIndex].GetRotation();
-			const float QuatSign = Rotation.W < 0.f ? -1.f : 1.f;
-			OutputBuffer[Index++] = Rotation.X * QuatSign;
-			OutputBuffer[Index++] = Rotation.Y * QuatSign;
-			OutputBuffer[Index++] = Rotation.Z * QuatSign;
-			OutputBuffer[Index++] = Rotation.W * QuatSign;
-		}
-		else
-		{
-			// Use an identity quaternion as input.
-			OutputBuffer[Index++] = 0.0f;
-			OutputBuffer[Index++] = 0.0f;
-			OutputBuffer[Index++] = 0.0f;
-			OutputBuffer[Index++] = 1.0f;
-		}
+		const FMatrix RotationMatrix = (SkelMeshBoneIndex != INDEX_NONE) ? BoneTransforms[SkelMeshBoneIndex].GetRotation().ToMatrix() : FMatrix::Identity;
+		const FVector X = RotationMatrix.GetColumn(0);
+		const FVector Y = RotationMatrix.GetColumn(1);	
+		OutputBuffer[Index++] = X.X;
+		OutputBuffer[Index++] = X.Y;
+		OutputBuffer[Index++] = X.Z;
+		OutputBuffer[Index++] = Y.X;
+		OutputBuffer[Index++] = Y.Y;
+		OutputBuffer[Index++] = Y.Z;
 	}
 
 	return Index;
diff --git a/Engine/Plugins/Experimental/MLDeformer/Source/MLDeformer/Public/MLDeformerInputInfo.h b/Engine/Plugins/Experimental/MLDeformer/Source/MLDeformer/Public/MLDeformerInputInfo.h
--- a/Engine/Plugins/Experimental/MLDeformer/Source/MLDeformer/Public/MLDeformerInputInfo.h
+++ b/Engine/Plugins/Experimental/MLDeformer/Source/MLDeformer/Public/MLDeformerInputInfo.h
@@ -121,7 +121,7 @@ public:
 
 	/**
 	 * Extract bone space rotations, as a float array.
-	 * The number of output rotations are NumBones * 4, where the array contains values like: xyzw,xyzw,xyzw,... where each xyzw are the 
+	 * The number of output rotations are NumBones * 6, where the array contains a set two columns of a 3x3 rotation matrix values.
 	 * components of the bone's bone space (local space) rotation quaternion.
 	 * @param SkelMeshComponent The skeletal mesh component to sample from.
 	 * @param OutRotations The output rotation values. This array will be resized internally.
