[MLDeformer] Changed the maximum value of the Vertex Delta Multiplier to 1, instead of 5, as that is more practical.

#trivial

diff --git a/Engine/Plugins/Experimental/MLDeformer/Source/MLDeformer/Public/MLDeformerVizSettings.h b/Engine/Plugins/Experimental/MLDeformer/Source/MLDeformer/Public/MLDeformerVizSettings.h
--- a/Engine/Plugins/Experimental/MLDeformer/Source/MLDeformer/Public/MLDeformerVizSettings.h
+++ b/Engine/Plugins/Experimental/MLDeformer/Source/MLDeformer/Public/MLDeformerVizSettings.h
@@ -95,7 +95,7 @@ public:
 	TObjectPtr<UGeometryCache> GroundTruth = nullptr;
 
 	/** The scale factor of the ML deformer deltas being applied on top of the linear skinned results. */
-	UPROPERTY(EditAnywhere, Category = "Live Settings", meta = (ClampMin = "0.0", ClampMax = "5.0"))
+	UPROPERTY(EditAnywhere, Category = "Live Settings", meta = (ClampMin = "0.0", ClampMax = "1.0"))
 	float VertexDeltaMultiplier = 1.0f;
 
 	/** The play speed factor of the test anim sequence. */
