[MLDeformer] Crash fix for the Mac.

diff --git a/Engine/Plugins/Experimental/MLDeformer/Source/MLDeformer/Private/MLDeformerAsset.cpp b/Engine/Plugins/Experimental/MLDeformer/Source/MLDeformer/Private/MLDeformerAsset.cpp
--- a/Engine/Plugins/Experimental/MLDeformer/Source/MLDeformer/Private/MLDeformerAsset.cpp
+++ b/Engine/Plugins/Experimental/MLDeformer/Source/MLDeformer/Private/MLDeformerAsset.cpp
@@ -86,6 +86,13 @@ void UMLDeformerAsset::PostLoad()
 	}
 
 	InitGPUData();
+
+	if (NeuralNetwork && 
+		(NeuralNetwork->GetDeviceType() != ENeuralDeviceType::GPU || 
+		NeuralNetwork->GetOutputDeviceType() != ENeuralDeviceType::GPU))
+	{
+		UE_LOG(LogMLDeformer, Error, TEXT("Neural net in ML Deformer '%s' cannot run on the GPU, it will not be active."), *GetName());
+	}
 }
 
 // Used for the FBoenReference, so it knows what skeleton to pick bones from.
diff --git a/Engine/Plugins/Experimental/MLDeformer/Source/MLDeformer/Private/MLDeformerDataInterface.cpp b/Engine/Plugins/Experimental/MLDeformer/Source/MLDeformer/Private/MLDeformerDataInterface.cpp
--- a/Engine/Plugins/Experimental/MLDeformer/Source/MLDeformer/Private/MLDeformerDataInterface.cpp
+++ b/Engine/Plugins/Experimental/MLDeformer/Source/MLDeformer/Private/MLDeformerDataInterface.cpp
@@ -95,13 +95,21 @@ bool UMLDeformerDataProvider::IsValid() const
 	}
 	
 	UMLDeformerComponent* DeformerComponent = SkeletalMeshComponent->GetOwner()->FindComponentByClass<UMLDeformerComponent>();
+	UMLDeformerAsset* DeformerAsset = DeformerComponent ? DeformerComponent->GetDeformerAsset() : nullptr;
+	UNeuralNetwork* NeuralNetwork = (DeformerAsset != nullptr && DeformerComponent != nullptr) ? DeformerComponent->GetDeformerAsset()->GetInferenceNeuralNetwork() : nullptr;
+	if (NeuralNetwork)
+	{
+		if (!NeuralNetwork->IsLoaded() ||
+			NeuralNetwork->GetDeviceType() != ENeuralDeviceType::GPU || 
+			NeuralNetwork->GetOutputDeviceType() != ENeuralDeviceType::GPU)
+		{
+			return false;
+		}
+	}
 
-	return
-		DeformerComponent != nullptr &&
-		DeformerComponent->GetDeformerAsset() != nullptr &&
-		DeformerComponent->GetDeformerAsset()->GetVertexMapBuffer().ShaderResourceViewRHI != nullptr &&
-		DeformerComponent->GetDeformerAsset()->GetInferenceNeuralNetwork() != nullptr &&
-		DeformerComponent->GetDeformerAsset()->GetInferenceNeuralNetwork()->IsLoaded() &&
+	return	
+		NeuralNetwork != nullptr &&
+		DeformerAsset->GetVertexMapBuffer().ShaderResourceViewRHI != nullptr &&
 		DeformerComponent->GetDeformerInstance().GetNeuralNetworkInferenceHandle() != -1;
 }
 
diff --git a/Engine/Plugins/Experimental/MLDeformer/Source/MLDeformer/Private/MLDeformerDebugDataInterface.cpp b/Engine/Plugins/Experimental/MLDeformer/Source/MLDeformer/Private/MLDeformerDebugDataInterface.cpp
--- a/Engine/Plugins/Experimental/MLDeformer/Source/MLDeformer/Private/MLDeformerDebugDataInterface.cpp
+++ b/Engine/Plugins/Experimental/MLDeformer/Source/MLDeformer/Private/MLDeformerDebugDataInterface.cpp
@@ -18,6 +18,7 @@
 #include "RenderGraphResources.h"
 #include "Rendering/SkeletalMeshModel.h"
 #include "SkeletalRenderPublic.h"
+#include "NeuralNetwork.h"
 
 FString UMLDeformerDebugDataInterface::GetDisplayName() const
 {
@@ -162,10 +163,22 @@ UComputeDataProvider* UMLDeformerDebugDataInterface::CreateDataProvider(TArrayVi
 bool UMLDeformerDebugDataProvider::IsValid() const
 {
 #if WITH_EDITORONLY_DATA
+	UMLDeformerComponent* DeformerComponent = SkeletalMeshComponent->GetOwner()->FindComponentByClass<UMLDeformerComponent>();
+	UNeuralNetwork* NeuralNetwork = (DeformerAsset != nullptr && DeformerComponent != nullptr) ? DeformerAsset->GetInferenceNeuralNetwork() : nullptr;
+	if (NeuralNetwork)
+	{
+		if (!NeuralNetwork->IsLoaded() ||
+			NeuralNetwork->GetDeviceType() != ENeuralDeviceType::GPU || 
+			NeuralNetwork->GetOutputDeviceType() != ENeuralDeviceType::GPU)
+		{
+			return false;
+		}
+	}
+
 	return
+		NeuralNetwork != nullptr &&
 		SkeletalMeshComponent != nullptr &&
 		SkeletalMeshComponent->MeshObject != nullptr &&
-		DeformerAsset != nullptr &&
 		DeformerAsset->GetVertexMapBuffer().ShaderResourceViewRHI != nullptr;
 #else
 	return false; // This data interface is only valid in editor.
diff --git a/Engine/Plugins/Experimental/MLDeformer/Source/MLDeformer/Private/MLDeformerInstance.cpp b/Engine/Plugins/Experimental/MLDeformer/Source/MLDeformer/Private/MLDeformerInstance.cpp
--- a/Engine/Plugins/Experimental/MLDeformer/Source/MLDeformer/Private/MLDeformerInstance.cpp
+++ b/Engine/Plugins/Experimental/MLDeformer/Source/MLDeformer/Private/MLDeformerInstance.cpp
@@ -193,10 +193,17 @@ void FMLDeformerInstance::Update()
 		return;
 	}
 
-	// We only support GPU processing of the neural network at the moment.
+	
+	// If we're not on the GPU we can't continue really.
+	// This is needed as the deformer graph system needs it on the GPU.
+	// Some platforms might not support GPU yet.
+	// Only the inputs are on the CPU.
 	check(NeuralNetwork->GetInputDeviceType() == ENeuralDeviceType::CPU);
-	check(NeuralNetwork->GetDeviceType() == ENeuralDeviceType::GPU);
-	check(NeuralNetwork->GetOutputDeviceType() == ENeuralDeviceType::GPU);
+	if (NeuralNetwork->GetDeviceType() != ENeuralDeviceType::GPU ||
+		NeuralNetwork->GetOutputDeviceType() != ENeuralDeviceType::GPU)
+	{
+		return;
+	}
 
 	// Allocate an inference context if none has already been allocated.
 	if (NeuralNetworkInferenceHandle == -1)
